Fix crash on new level due to menu system including references to deleted widgets
#rnx

diff --git a/Engine/Source/Developer/ToolMenus/Private/ToolMenus.cpp b/Engine/Source/Developer/ToolMenus/Private/ToolMenus.cpp
--- a/Engine/Source/Developer/ToolMenus/Private/ToolMenus.cpp
+++ b/Engine/Source/Developer/ToolMenus/Private/ToolMenus.cpp
@@ -1160,6 +1160,28 @@ void UToolMenus::AddReferencedObjects(UObject* InThis, FReferenceCollector& Coll
 		}
 	}
 
+	for (auto WidgetsForMenuNameIt = This->GeneratedMenuWidgets.CreateIterator(); WidgetsForMenuNameIt; ++WidgetsForMenuNameIt)
+	{
+		FGeneratedToolMenuWidgets& WidgetsForMenuName = WidgetsForMenuNameIt->Value;
+
+		for (auto Instance = WidgetsForMenuName.Instances.CreateIterator(); Instance; ++Instance)
+		{
+			if (Instance->Widget.IsValid())
+			{
+				Collector.AddReferencedObject(Instance->GeneratedMenu, InThis);
+			}
+			else
+			{
+				Instance.RemoveCurrent();
+			}
+		}
+
+		if (WidgetsForMenuName.Instances.Num() == 0)
+		{
+			WidgetsForMenuNameIt.RemoveCurrent();
+		}
+	}
+
 	Super::AddReferencedObjects(InThis, Collector);
 }
 
diff --git a/Engine/Source/Developer/ToolMenus/Public/ToolMenus.h b/Engine/Source/Developer/ToolMenus/Public/ToolMenus.h
--- a/Engine/Source/Developer/ToolMenus/Public/ToolMenus.h
+++ b/Engine/Source/Developer/ToolMenus/Public/ToolMenus.h
@@ -32,23 +32,14 @@ class UToolMenuEntryScript;
 struct FToolMenuEntry;
 struct FToolMenuSection;
 
-USTRUCT()
 struct FGeneratedToolMenuWidget
 {
-	GENERATED_BODY()
-
-	UPROPERTY()
 	UToolMenu* GeneratedMenu;
-
 	TWeakPtr<SWidget> Widget;
 };
 
-USTRUCT()
 struct FGeneratedToolMenuWidgets
 {
-	GENERATED_BODY()
-
-	UPROPERTY()
 	TArray<FGeneratedToolMenuWidget> Instances;
 };
 
@@ -342,7 +333,6 @@ private:
 	UPROPERTY()
 	TMap<FName, UToolMenu*> Menus;
 
-	UPROPERTY()
 	TMap<FName, FGeneratedToolMenuWidgets> GeneratedMenuWidgets;
 
 	TMap<TWeakPtr<FMultiBox>, TArray<UObject*>> WidgetObjectReferences;
