Add tool menus modules
#rnx

diff --git a/Engine/Source/Developer/ToolMenus/Private/ToolMenu.cpp b/Engine/Source/Developer/ToolMenus/Private/ToolMenu.cpp
--- a/Engine/Source/Developer/ToolMenus/Private/ToolMenu.cpp
+++ b/Engine/Source/Developer/ToolMenus/Private/ToolMenu.cpp
@@ -18,6 +18,7 @@ UToolMenu::UToolMenu() :
 	, bToolBarIsFocusable(false)
 	, bToolBarForceSmallIcons(false)
 	, bRegistered(false)
+	, bIsRegistering(false)
 	, StyleSet(&FCoreStyle::Get())
 	, MaxHeight(INT_MAX)
 {
@@ -120,6 +121,11 @@ FToolMenuSection& UToolMenu::AddDynamicSection(const FName SectionName, const FN
 	return Section;
 }
 
+bool UToolMenu::IsRegistering() const
+{
+	return bIsRegistering;
+}
+
 FToolMenuSection& UToolMenu::AddSection(const FName SectionName, const TAttribute< FText >& InLabel, const FToolMenuInsert InPosition)
 {
 	for (FToolMenuSection& Section : Sections)
@@ -142,7 +148,7 @@ FToolMenuSection& UToolMenu::AddSection(const FName SectionName, const TAttribut
 
 	FToolMenuSection& NewSection = Sections.AddDefaulted_GetRef();
 	NewSection.InitSection(SectionName, InLabel, InPosition);
-	//NewSection.OwnerMenu = this;
+	NewSection.bIsRegistering = IsRegistering();
 	return NewSection;
 }
 
diff --git a/Engine/Source/Developer/ToolMenus/Private/ToolMenuEntry.cpp b/Engine/Source/Developer/ToolMenus/Private/ToolMenuEntry.cpp
--- a/Engine/Source/Developer/ToolMenus/Private/ToolMenuEntry.cpp
+++ b/Engine/Source/Developer/ToolMenus/Private/ToolMenuEntry.cpp
@@ -15,7 +15,8 @@ FToolMenuEntry::FToolMenuEntry() :
 	Type(EMultiBlockType::None),
 	UserInterfaceActionType(EUserInterfaceActionType::Button),
 	bShouldCloseWindowAfterMenuSelection(true),
-	ScriptObject(nullptr)
+	ScriptObject(nullptr),
+	bAddedDuringRegister(false)
 {
 }
 
@@ -25,7 +26,8 @@ FToolMenuEntry::FToolMenuEntry(const FToolMenuOwner InOwner, const FName InName,
 	Type(InType),
 	UserInterfaceActionType(EUserInterfaceActionType::Button),
 	bShouldCloseWindowAfterMenuSelection(true),
-	ScriptObject(nullptr)
+	ScriptObject(nullptr),
+	bAddedDuringRegister(false)
 {
 }
 
diff --git a/Engine/Source/Developer/ToolMenus/Private/ToolMenuSection.cpp b/Engine/Source/Developer/ToolMenus/Private/ToolMenuSection.cpp
--- a/Engine/Source/Developer/ToolMenus/Private/ToolMenuSection.cpp
+++ b/Engine/Source/Developer/ToolMenus/Private/ToolMenuSection.cpp
@@ -11,7 +11,8 @@
 #include "Internationalization/Internationalization.h"
 
 FToolMenuSection::FToolMenuSection() :
-	ToolMenuSectionDynamic(nullptr)
+	ToolMenuSectionDynamic(nullptr),
+	bIsRegistering(false)
 {
 
 }
@@ -33,22 +34,32 @@ void FToolMenuSection::InitGeneratedSectionCopy(const FToolMenuSection& Source,
 	Context = InContext;
 }
 
+bool FToolMenuSection::IsRegistering() const
+{
+	return bIsRegistering;
+}
+
 FToolMenuEntry& FToolMenuSection::AddEntry(const FToolMenuEntry& Args)
 {
 	if (Args.Name == NAME_None)
 	{
-		return Blocks.Add_GetRef(Args);
+		FToolMenuEntry& Result = Blocks.Add_GetRef(Args);
+		Result.bAddedDuringRegister = IsRegistering();
+		return Result;
 	}
 
 	int32 BlockIndex = IndexOfBlock(Args.Name);
 	if (BlockIndex != INDEX_NONE)
 	{
 		Blocks[BlockIndex] = Args;
+		Blocks[BlockIndex].bAddedDuringRegister = IsRegistering();
 		return Blocks[BlockIndex];
 	}
 	else
 	{
-		return Blocks.Add_GetRef(Args);
+		FToolMenuEntry& Result = Blocks.Add_GetRef(Args);
+		Result.bAddedDuringRegister = IsRegistering();
+		return Result;
 	}
 }
 
diff --git a/Engine/Source/Developer/ToolMenus/Private/ToolMenus.cpp b/Engine/Source/Developer/ToolMenus/Private/ToolMenus.cpp
--- a/Engine/Source/Developer/ToolMenus/Private/ToolMenus.cpp
+++ b/Engine/Source/Developer/ToolMenus/Private/ToolMenus.cpp
@@ -15,6 +15,8 @@
 
 #define LOCTEXT_NAMESPACE "ToolMenuSubsystem"
 
+DEFINE_LOG_CATEGORY(LogToolMenus);
+
 UToolMenus* UToolMenus::Singleton = nullptr;
 bool UToolMenus::bHasShutDown = false;
 FSimpleMulticastDelegate UToolMenus::StartupCallbacks;
@@ -201,13 +203,21 @@ void UToolMenus::AssembleMenuSection(UToolMenu* GeneratedMenu, const UToolMenu*
 {
 	// Build list of blocks in expected order including blocks created by construct delegates
 	TArray<FToolMenuEntry> RemainingBlocks;
+	TArray<FToolMenuEntry> BlocksToAddLast;
 
 	UToolMenu* ConstructedEntries = nullptr;
 	for (const FToolMenuEntry& Block : OtherSection.Blocks)
 	{
 		if (!Block.IsNonLegacyDynamicConstruct())
 		{
-			RemainingBlocks.Add(Block);
+			if (Block.bAddedDuringRegister)
+			{
+				RemainingBlocks.Add(Block);
+			}
+			else
+			{
+				BlocksToAddLast.Add(Block);
+			}
 			continue;
 		}
 
@@ -263,7 +273,14 @@ void UToolMenus::AssembleMenuSection(UToolMenu* GeneratedMenu, const UToolMenu*
 			}
 			else
 			{
-				RemainingBlocks.Add(GeneratedEntry);
+				if (GeneratedEntry.bAddedDuringRegister)
+				{
+					RemainingBlocks.Add(GeneratedEntry);
+				}
+				else
+				{
+					BlocksToAddLast.Add(GeneratedEntry);
+				}
 				GeneratedEntries.RemoveAt(0, 1, false);
 			}
 		}
@@ -275,6 +292,8 @@ void UToolMenus::AssembleMenuSection(UToolMenu* GeneratedMenu, const UToolMenu*
 		ConstructedEntries = nullptr;
 	}
 
+	RemainingBlocks.Append(BlocksToAddLast);
+
 	// Repeatedly loop because insert location may not exist until later in list
 	while (RemainingBlocks.Num() > 0)
 	{
@@ -710,9 +729,14 @@ void UToolMenus::PopulateSubMenu(FMenuBuilder& MenuBuilder, TWeakObjectPtr<UTool
 	}
 }
 
-TSharedRef<SWidget> UToolMenus::GenerateToolbarComboButtonMenu(const FName SubMenuFullName, FToolMenuContext InContext)
+TSharedRef<SWidget> UToolMenus::GenerateToolbarComboButtonMenu(TWeakObjectPtr<UToolMenu> InParent, const FName InBlockName)
 {
-	return GenerateWidget(SubMenuFullName, InContext);
+	if (UToolMenu* GeneratedMenu = GenerateSubMenu(InParent.Get(), InBlockName))
+	{
+		return GenerateWidget(GeneratedMenu);
+	}
+
+	return SNullWidget::NullWidget;
 }
 
 void UToolMenus::FillMenuBarDropDown(class FMenuBuilder& MenuBuilder, FName InParentName, FName InChildName, FToolMenuContext InMenuContext)
@@ -833,14 +857,32 @@ void UToolMenus::PopulateMenuBuilder(FMenuBuilder& MenuBuilder, UToolMenu* MenuD
 					}
 					else
 					{
-						MenuBuilder.AddSubMenu(
-							Block.Label,
-							Block.ToolTip,
-							NewMenuDelegate,
-							Block.SubMenuData.bOpenSubMenuOnClick,
-							Block.Icon.Get(),
-							Block.bShouldCloseWindowAfterMenuSelection,
-							Block.Name);
+						if (UIAction.IsBound())
+						{
+							MenuBuilder.AddSubMenu(
+								Block.Label,
+								Block.ToolTip,
+								NewMenuDelegate,
+								UIAction,
+								Block.Name,
+								Block.UserInterfaceActionType,
+								Block.SubMenuData.bOpenSubMenuOnClick,
+								Block.Icon.Get(),
+								Block.bShouldCloseWindowAfterMenuSelection
+							);
+						}
+						else
+						{
+							MenuBuilder.AddSubMenu(
+								Block.Label,
+								Block.ToolTip,
+								NewMenuDelegate,
+								Block.SubMenuData.bOpenSubMenuOnClick,
+								Block.Icon.Get(),
+								Block.bShouldCloseWindowAfterMenuSelection,
+								Block.Name
+							);
+						}
 					}
 				}
 				else
@@ -975,7 +1017,7 @@ void UToolMenus::PopulateToolBarBuilder(FToolBarBuilder& ToolBarBuilder, UToolMe
 					}
 					else
 					{
-						UE_LOG(LogToolMenus, Error, TEXT("UI command not found for toolbar entry: %s, toolbar: %s"), *Block.Name.ToString(), *MenuData->MenuName.ToString());
+						UE_LOG(LogToolMenus, Verbose, TEXT("UI command not found for toolbar entry: %s, toolbar: %s"), *Block.Name.ToString(), *MenuData->MenuName.ToString());
 					}
 
 					ToolBarBuilder.AddToolBarButton(Block.Command, Block.Name, Block.Label, Block.ToolTip, Block.Icon, Block.TutorialHighlightName);
@@ -1006,7 +1048,7 @@ void UToolMenus::PopulateToolBarBuilder(FToolBarBuilder& ToolBarBuilder, UToolMe
 				else
 				{
 					FName SubMenuFullName = JoinMenuPaths(MenuData->MenuName, Block.Name);
-					FOnGetContent Delegate = FOnGetContent::CreateUObject(this, &UToolMenus::GenerateToolbarComboButtonMenu, SubMenuFullName, MenuData->Context);
+					FOnGetContent Delegate = FOnGetContent::CreateUObject(this, &UToolMenus::GenerateToolbarComboButtonMenu, TWeakObjectPtr<UToolMenu>(MenuData), Block.Name);
 					ToolBarBuilder.AddComboButton(UIAction, Delegate, Block.Label, Block.ToolTip, Block.Icon, Block.ToolBarData.bSimpleComboBox, Block.TutorialHighlightName);
 				}
 			}
@@ -1665,6 +1707,11 @@ UToolMenu* UToolMenus::RegisterMenu(const FName InName, const FName InParent, EM
 			Found->MenuType = InType;
 			Found->MenuOwner = CurrentOwner();
 			Found->bRegistered = true;
+			Found->bIsRegistering = true;
+			for (FToolMenuSection& Section : Found->Sections)
+			{
+				Section.bIsRegistering = Found->bIsRegistering;
+			}
 		}
 		else if (bWarnIfAlreadyRegistered)
 		{
@@ -1677,6 +1724,7 @@ UToolMenu* UToolMenus::RegisterMenu(const FName InName, const FName InParent, EM
 	UToolMenu* ToolMenu = NewObject<UToolMenu>(this);
 	ToolMenu->InitMenu(CurrentOwner(), InName, InParent, InType);
 	ToolMenu->bRegistered = true;
+	ToolMenu->bIsRegistering = true;
 	Menus.Add(InName, ToolMenu);
 	return ToolMenu;
 }
@@ -1685,12 +1733,18 @@ UToolMenu* UToolMenus::ExtendMenu(const FName InName)
 {
 	if (UToolMenu* Found = FindMenu(InName))
 	{
+		Found->bIsRegistering = false;
+		for (FToolMenuSection& Section : Found->Sections)
+		{
+			Section.bIsRegistering = Found->bIsRegistering;
+		}
 		return Found;
 	}
 
 	UToolMenu* ToolMenu = NewObject<UToolMenu>(this);
 	ToolMenu->MenuName = InName;
 	ToolMenu->bRegistered = false;
+	ToolMenu->bIsRegistering = false;
 	Menus.Add(InName, ToolMenu);
 	return ToolMenu;
 }
diff --git a/Engine/Source/Developer/ToolMenus/Public/IToolMenusModule.h b/Engine/Source/Developer/ToolMenus/Public/IToolMenusModule.h
--- a/Engine/Source/Developer/ToolMenus/Public/IToolMenusModule.h
+++ b/Engine/Source/Developer/ToolMenus/Public/IToolMenusModule.h
@@ -5,7 +5,7 @@
 #include "Modules/ModuleInterface.h"
 #include "Modules/ModuleManager.h"
 
-DEFINE_LOG_CATEGORY_STATIC(LogToolMenus, Log, All);
+DECLARE_LOG_CATEGORY_EXTERN(LogToolMenus, Log, All);
 
 class IToolMenusModule : public IModuleInterface
 {
diff --git a/Engine/Source/Developer/ToolMenus/Public/ToolMenu.h b/Engine/Source/Developer/ToolMenus/Public/ToolMenu.h
--- a/Engine/Source/Developer/ToolMenus/Public/ToolMenu.h
+++ b/Engine/Source/Developer/ToolMenus/Public/ToolMenu.h
@@ -82,7 +82,7 @@ public:
 	template <typename TContextType>
 	TContextType* FindContext() const
 	{
-		return Context.Find<TContextType>();
+		return Context.FindContext<TContextType>();
 	}
 
 	//~ Begin UObject Interface
@@ -105,6 +105,8 @@ private:
 
 	int32 FindInsertIndex(const FToolMenuSection& InSection) const;
 
+	bool IsRegistering() const;
+
 public:
 
 	UPROPERTY(EditAnywhere, BlueprintReadWrite, Category = "Tool Menus")
@@ -161,6 +163,7 @@ public:
 private:
 
 	bool bRegistered;
+	bool bIsRegistering;
 
 	const ISlateStyle* StyleSet;
 
diff --git a/Engine/Source/Developer/ToolMenus/Public/ToolMenuContext.h b/Engine/Source/Developer/ToolMenus/Public/ToolMenuContext.h
--- a/Engine/Source/Developer/ToolMenus/Public/ToolMenuContext.h
+++ b/Engine/Source/Developer/ToolMenus/Public/ToolMenuContext.h
@@ -37,6 +37,20 @@ public:
 	FToolMenuContext(UObject* InContext);
 	FToolMenuContext(TSharedPtr<FUICommandList> InCommandList, TSharedPtr<FExtender> InExtender = TSharedPtr<FExtender>(), UObject* InContext = nullptr);
 
+	template <typename TContextType>
+	TContextType* FindContext() const
+	{
+		for (UObject* Object : ContextObjects)
+		{
+			if (TContextType* Result = Cast<TContextType>(Object))
+			{
+				return Result;
+			}
+		}
+
+		return nullptr;
+	}
+
 	template <typename TContextType>
 	TContextType* Find() const
 	{
diff --git a/Engine/Source/Developer/ToolMenus/Public/ToolMenuEntry.h b/Engine/Source/Developer/ToolMenus/Public/ToolMenuEntry.h
--- a/Engine/Source/Developer/ToolMenus/Public/ToolMenuEntry.h
+++ b/Engine/Source/Developer/ToolMenus/Public/ToolMenuEntry.h
@@ -176,5 +176,7 @@ private:
 
 	FNewToolMenuSectionDelegate Construct;
 	FNewToolMenuDelegateLegacy ConstructLegacy;
+
+	bool bAddedDuringRegister;
 };
 
diff --git a/Engine/Source/Developer/ToolMenus/Public/ToolMenuSection.h b/Engine/Source/Developer/ToolMenus/Public/ToolMenuSection.h
--- a/Engine/Source/Developer/ToolMenus/Public/ToolMenuSection.h
+++ b/Engine/Source/Developer/ToolMenus/Public/ToolMenuSection.h
@@ -49,7 +49,7 @@ public:
 	template <typename TContextType>
 	TContextType* FindContext() const
 	{
-		return Context.Find<TContextType>();
+		return Context.FindContext<TContextType>();
 	}
 
 	FToolMenuEntry* FindEntry(const FName InName);
@@ -69,6 +69,8 @@ private:
 
 	bool IsNonLegacyDynamic() const;
 
+	bool IsRegistering() const;
+
 	friend class UToolMenuSectionExtensions;
 	friend class UToolMenus;
 	friend class UToolMenu;
@@ -93,4 +95,8 @@ public:
 	TAttribute<FText> Label;
 
 	FNewSectionConstructChoice Construct;
+
+private:
+
+	bool bIsRegistering;
 };
diff --git a/Engine/Source/Developer/ToolMenus/Public/ToolMenus.h b/Engine/Source/Developer/ToolMenus/Public/ToolMenus.h
--- a/Engine/Source/Developer/ToolMenus/Public/ToolMenus.h
+++ b/Engine/Source/Developer/ToolMenus/Public/ToolMenus.h
@@ -322,7 +322,7 @@ private:
 	void PopulateMenuBarBuilder(FMenuBarBuilder& MenuBarBuilder, UToolMenu* MenuData);
 	void PopulateToolBarBuilder(FToolBarBuilder& ToolBarBuilder, UToolMenu* MenuData);
 
-	TSharedRef<SWidget> GenerateToolbarComboButtonMenu(const FName SubMenuFullName, FToolMenuContext InContext);
+	TSharedRef<SWidget> GenerateToolbarComboButtonMenu(TWeakObjectPtr<UToolMenu> InParent, const FName InBlockName);
 
 	FOnGetContent ConvertWidgetChoice(const FNewToolMenuWidgetChoice& Choice, const FToolMenuContext& Context) const;
 
diff --git a/Engine/Source/Editor/GraphEditor/Private/SGraphEditorImpl.cpp b/Engine/Source/Editor/GraphEditor/Private/SGraphEditorImpl.cpp
--- a/Engine/Source/Editor/GraphEditor/Private/SGraphEditorImpl.cpp
+++ b/Engine/Source/Editor/GraphEditor/Private/SGraphEditorImpl.cpp
@@ -675,7 +675,7 @@ void SGraphEditorImpl::RegisterContextMenu(const UEdGraphSchema* Schema, FToolMe
 {
 	UToolMenus* ToolMenus = UToolMenus::Get();
 	const FName SchemaMenuName = Schema->GetContextMenuName();
-	UGraphNodeContextMenuContext* Context = MenuContext.Find<UGraphNodeContextMenuContext>();
+	UGraphNodeContextMenuContext* Context = MenuContext.FindContext<UGraphNodeContextMenuContext>();
 
 	// Root menu
 	// "GraphEditor.GraphContextMenu.Common"
@@ -771,7 +771,7 @@ UToolMenu* SGraphEditorImpl::GenerateContextMenu(const UEdGraphSchema* Schema, F
 	RegisterContextMenu(Schema, MenuContext);
 
 	UToolMenus* ToolMenus = UToolMenus::Get();
-	UGraphNodeContextMenuContext* Context = MenuContext.Find<UGraphNodeContextMenuContext>();
+	UGraphNodeContextMenuContext* Context = MenuContext.FindContext<UGraphNodeContextMenuContext>();
 
 	FName MenuName = NAME_None;
 	if (Context->Node)
