Menus can now be replaced using ini files

diff --git a/Engine/Source/Developer/ToolMenus/Private/ToolMenus.cpp b/Engine/Source/Developer/ToolMenus/Private/ToolMenus.cpp
--- a/Engine/Source/Developer/ToolMenus/Private/ToolMenus.cpp
+++ b/Engine/Source/Developer/ToolMenus/Private/ToolMenus.cpp
@@ -1423,7 +1423,14 @@ UToolMenu* UToolMenus::GenerateMenuOrSubMenuForEdit(const UToolMenu* InMenu)
 
 UToolMenu* UToolMenus::GenerateMenu(const FName Name, const FToolMenuContext& InMenuContext)
 {
-	return GenerateMenuFromHierarchy(CollectHierarchy(Name), InMenuContext);
+	// Allow substituting one menu for another
+	FName MenuName = Name;
+	if (const FName* Found = MenuSubstitutionsDuringGenerate.Find(Name))
+	{
+		MenuName = *Found;
+	}
+
+	return GenerateMenuFromHierarchy(CollectHierarchy(MenuName), InMenuContext);
 }
 
 UToolMenu* UToolMenus::GenerateMenuFromHierarchy(const TArray<UToolMenu*>& Hierarchy, const FToolMenuContext& InMenuContext)
diff --git a/Engine/Source/Developer/ToolMenus/Public/ToolMenus.h b/Engine/Source/Developer/ToolMenus/Public/ToolMenus.h
--- a/Engine/Source/Developer/ToolMenus/Public/ToolMenus.h
+++ b/Engine/Source/Developer/ToolMenus/Public/ToolMenus.h
@@ -359,6 +359,10 @@ private:
 	UPROPERTY(config, EditAnywhere, Category = Misc)
 	TArray<FCustomizedToolMenu> CustomizedMenus;
 
+	/* Allow substituting one menu for another during generate but not during find or extend */
+	UPROPERTY(config, EditAnywhere, Category = Misc)
+	TMap<FName, FName> MenuSubstitutionsDuringGenerate;
+
 	UPROPERTY()
 	TMap<FName, UToolMenu*> Menus;
 
