Fixed dynamic tool menu entries appearing at the end of a section
#rnx

diff --git a/Engine/Source/Developer/ToolMenus/Private/ToolMenus.cpp b/Engine/Source/Developer/ToolMenus/Private/ToolMenus.cpp
--- a/Engine/Source/Developer/ToolMenus/Private/ToolMenus.cpp
+++ b/Engine/Source/Developer/ToolMenus/Private/ToolMenus.cpp
@@ -273,7 +273,7 @@ void UToolMenus::AssembleMenuSection(UToolMenu* GeneratedMenu, const UToolMenu*
 			}
 			else
 			{
-				if (GeneratedEntry.bAddedDuringRegister)
+				if (Block.bAddedDuringRegister)
 				{
 					RemainingBlocks.Add(GeneratedEntry);
 				}
