Tool menus can be customized with whitelists

diff --git a/Engine/Source/Developer/ToolMenus/Private/ToolMenus.cpp b/Engine/Source/Developer/ToolMenus/Private/ToolMenus.cpp
--- a/Engine/Source/Developer/ToolMenus/Private/ToolMenus.cpp
+++ b/Engine/Source/Developer/ToolMenus/Private/ToolMenus.cpp
@@ -597,6 +597,23 @@ void UToolMenus::ApplyCustomization(UToolMenu* GeneratedMenu)
 		}
 	}
 
+	// Hide items based on whitelist
+	if (CustomizedMenu.WhitelistEntries.Num() > 0)
+	{
+		for (int32 SectionIndex = 0; SectionIndex < NewSections.Num(); ++SectionIndex)
+		{
+			FToolMenuSection& Section = NewSections[SectionIndex];
+			for (int32 i = 0; i < Section.Blocks.Num(); ++i)
+			{
+				if (!CustomizedMenu.IsEntryWhitelisted(Section.Blocks[i].Name))
+				{
+					Section.Blocks.RemoveAt(i);
+					--i;
+				}
+			}
+		}
+	}
+
 	// Hide sections and entries
 	if (!GeneratedMenu->IsEditing())
 	{
@@ -1421,6 +1438,11 @@ UToolMenu* UToolMenus::GenerateMenuOrSubMenuForEdit(const UToolMenu* InMenu)
 	return nullptr;
 }
 
+void UToolMenus::AddMenuSubstitutionDuringGenerate(const FName OriginalMenu, const FName NewMenu)
+{
+	MenuSubstitutionsDuringGenerate.Add(OriginalMenu, NewMenu);
+}
+
 UToolMenu* UToolMenus::GenerateMenu(const FName Name, const FToolMenuContext& InMenuContext)
 {
 	// Allow substituting one menu for another
diff --git a/Engine/Source/Developer/ToolMenus/Public/ToolMenus.h b/Engine/Source/Developer/ToolMenus/Public/ToolMenus.h
--- a/Engine/Source/Developer/ToolMenus/Public/ToolMenus.h
+++ b/Engine/Source/Developer/ToolMenus/Public/ToolMenus.h
@@ -260,6 +260,9 @@ public:
 	/** Enables adding command to open edit menu dialog to each menu */
 	void SetEditMenusMode(bool bEnable);
 
+	/* Substitute one menu for another during generate but not during find or extend */
+	void AddMenuSubstitutionDuringGenerate(const FName OriginalMenu, const FName NewMenu);
+
 	/** Displaying extension points is for debugging menus */
 	DECLARE_DELEGATE_RetVal(bool, FShouldDisplayExtensionPoints);
 	FShouldDisplayExtensionPoints ShouldDisplayExtensionPoints;
diff --git a/Engine/Source/Runtime/Slate/Private/Framework/MultiBox/ToolMenuCustomization.cpp b/Engine/Source/Runtime/Slate/Private/Framework/MultiBox/ToolMenuCustomization.cpp
--- a/Engine/Source/Runtime/Slate/Private/Framework/MultiBox/ToolMenuCustomization.cpp
+++ b/Engine/Source/Runtime/Slate/Private/Framework/MultiBox/ToolMenuCustomization.cpp
@@ -32,6 +32,16 @@ FCustomizedToolMenuEntry* FCustomizedToolMenu::AddEntry(const FName InEntryName)
 	return &Entries.FindOrAdd(InEntryName);
 }
 
+void FCustomizedToolMenu::AddWhitelistEntry(const FName InEntryName)
+{
+	WhitelistEntries.AddUnique(InEntryName);
+}
+
+bool FCustomizedToolMenu::IsEntryWhitelisted(const FName InEntryName) const
+{
+	return WhitelistEntries.Contains(InEntryName);
+}
+
 ECustomizedToolMenuVisibility FCustomizedToolMenu::GetEntryVisiblity(const FName InEntryName) const
 {
 	if (const FCustomizedToolMenuEntry* Found = FindEntry(InEntryName))
@@ -193,6 +203,11 @@ FCustomizedToolMenu FCustomizedToolMenuHierarchy::GenerateFlattened() const
 				Result.Sections.Add(SectionIterator.Key, SectionIterator.Value);
 			}
 		}
+
+		for (const FName WhitelistEntry : Current->WhitelistEntries)
+		{
+			Result.AddWhitelistEntry(WhitelistEntry);
+		}
 	}
 
 	return Result;
diff --git a/Engine/Source/Runtime/Slate/Public/Framework/MultiBox/ToolMenuBase.h b/Engine/Source/Runtime/Slate/Public/Framework/MultiBox/ToolMenuBase.h
--- a/Engine/Source/Runtime/Slate/Public/Framework/MultiBox/ToolMenuBase.h
+++ b/Engine/Source/Runtime/Slate/Public/Framework/MultiBox/ToolMenuBase.h
@@ -73,12 +73,17 @@ struct SLATE_API FCustomizedToolMenu
 	UPROPERTY()
 	TArray<FName> SectionOrder;
 
+	UPROPERTY()
+	TArray<FName> WhitelistEntries;
+
 	FCustomizedToolMenuEntry* FindEntry(const FName InEntryName);
 	const FCustomizedToolMenuEntry* FindEntry(const FName InEntryName) const;
 	FCustomizedToolMenuEntry* AddEntry(const FName InEntryName);
 	ECustomizedToolMenuVisibility GetEntryVisiblity(const FName InSectionName) const;
 	bool IsEntryHidden(const FName InEntryName) const;
 	FName GetEntrySectionName(const FName InEntryName) const;
+	void AddWhitelistEntry(const FName InEntryName);
+	bool IsEntryWhitelisted(const FName InEntryName) const;
 
 	FCustomizedToolMenuSection* FindSection(const FName InSectionName);
 	const FCustomizedToolMenuSection* FindSection(const FName InSectionName) const;
