Tool Menus: Add option to disable menu extensions
#rnx

diff --git a/Engine/Source/Developer/ToolMenus/Private/ToolMenu.cpp b/Engine/Source/Developer/ToolMenus/Private/ToolMenu.cpp
--- a/Engine/Source/Developer/ToolMenus/Private/ToolMenu.cpp
+++ b/Engine/Source/Developer/ToolMenus/Private/ToolMenu.cpp
@@ -19,6 +19,7 @@ UToolMenu::UToolMenu() :
 	, bToolBarForceSmallIcons(false)
 	, bRegistered(false)
 	, bIsRegistering(false)
+	, bExtendersEnabled(true)
 	, StyleSet(&FCoreStyle::Get())
 	, MaxHeight(INT_MAX)
 {
@@ -51,6 +52,8 @@ void UToolMenu::InitGeneratedCopy(const UToolMenu* Source, const FName InMenuNam
 
 	SubMenuParent = Source->SubMenuParent;
 	SubMenuSourceEntryName = Source->SubMenuSourceEntryName;
+	bExtendersEnabled = Source->bExtendersEnabled;
+
 	MaxHeight = Source->MaxHeight;
 	if (InContext)
 	{
@@ -451,4 +454,9 @@ FString UToolMenu::GetSubMenuNamePath() const
 	}
 
 	return SubMenuNamePath;
-}
\ No newline at end of file
+}
+
+void UToolMenu::SetExtendersEnabled(bool bEnabled)
+{
+	bExtendersEnabled = bEnabled;
+}
diff --git a/Engine/Source/Developer/ToolMenus/Private/ToolMenus.cpp b/Engine/Source/Developer/ToolMenus/Private/ToolMenus.cpp
--- a/Engine/Source/Developer/ToolMenus/Private/ToolMenus.cpp
+++ b/Engine/Source/Developer/ToolMenus/Private/ToolMenus.cpp
@@ -681,13 +681,18 @@ UToolMenu* UToolMenus::GenerateSubMenu(const UToolMenu* InGeneratedParent, const
 		return nullptr;
 	}
 
+	FName SubMenuFullName = JoinMenuPaths(InGeneratedParent->GetMenuName(), InBlockName);
+	if (const FName* ReplacementName = MenuSubstitutionsDuringGenerate.Find(SubMenuFullName))
+	{
+		return GenerateMenuFromHierarchy(CollectHierarchy(*ReplacementName), InGeneratedParent->Context);
+	}
+
 	const FToolMenuEntry* Block = InGeneratedParent->FindEntry(InBlockName);
 	if (!Block)
 	{
 		return nullptr;
 	}
 
-	FName SubMenuFullName = JoinMenuPaths(InGeneratedParent->GetMenuName(), InBlockName);
 
 	// Submenus that are constructed by delegates can also be overridden by menus in the database
 	TArray<UToolMenu*> Hierarchy;
@@ -729,6 +734,14 @@ UToolMenu* UToolMenus::GenerateSubMenu(const UToolMenu* InGeneratedParent, const
 	{
 		UToolMenu* GeneratedMenu = NewObject<UToolMenu>(this);
 		GeneratedMenu->InitGeneratedCopy(Hierarchy[0], SubMenuFullName, &InGeneratedParent->Context);
+		for (UToolMenu* HiearchyItem : Hierarchy)
+		{
+			if (HiearchyItem && !HiearchyItem->bExtendersEnabled)
+			{
+				GeneratedMenu->SetExtendersEnabled(false);
+				break;
+			}
+		}
 		GeneratedMenu->SubMenuParent = InGeneratedParent;
 		GeneratedMenu->SubMenuSourceEntryName = InBlockName;
 		AssembleMenuHierarchy(GeneratedMenu, Hierarchy);
@@ -1443,6 +1456,35 @@ void UToolMenus::AddMenuSubstitutionDuringGenerate(const FName OriginalMenu, con
 	MenuSubstitutionsDuringGenerate.Add(OriginalMenu, NewMenu);
 }
 
+void UToolMenus::RemoveSubstitutionDuringGenerate(const FName InMenu)
+{
+	if (const FName* FoundOverrideMenuName = MenuSubstitutionsDuringGenerate.Find(InMenu))
+	{
+		const FName OverrideMenuName = *FoundOverrideMenuName;
+
+		// Update all active widget instances of this menu
+		FGeneratedToolMenuWidgets* OverrideMenuWidgets = GeneratedMenuWidgets.Find(OverrideMenuName);
+		if (OverrideMenuWidgets)
+		{
+			FGeneratedToolMenuWidgets* DestMenuWidgets = GeneratedMenuWidgets.Find(InMenu);
+			if (DestMenuWidgets)
+			{
+				DestMenuWidgets->Instances.Append(OverrideMenuWidgets->Instances);
+			}
+			else
+			{
+				GeneratedMenuWidgets.Add(InMenu, *OverrideMenuWidgets);
+			}
+
+			GeneratedMenuWidgets.Remove(OverrideMenuName);
+		}
+
+		MenuSubstitutionsDuringGenerate.Remove(InMenu);
+
+		CleanupStaleWidgetsNextTick();
+	}
+}
+
 UToolMenu* UToolMenus::GenerateMenu(const FName Name, const FToolMenuContext& InMenuContext)
 {
 	// Allow substituting one menu for another
@@ -1462,6 +1504,14 @@ UToolMenu* UToolMenus::GenerateMenuFromHierarchy(const TArray<UToolMenu*>& Hiera
 	if (Hierarchy.Num() > 0)
 	{
 		GeneratedMenu->InitGeneratedCopy(Hierarchy[0], Hierarchy.Last()->MenuName, &InMenuContext);
+		for (UToolMenu* HiearchyItem : Hierarchy)
+		{
+			if (HiearchyItem && !HiearchyItem->bExtendersEnabled)
+			{
+				GeneratedMenu->SetExtendersEnabled(false);
+				break;
+			}
+		}
 		AssembleMenuHierarchy(GeneratedMenu, Hierarchy);
 	}
 
@@ -1513,6 +1563,7 @@ TSharedRef<SWidget> UToolMenus::GenerateWidget(UToolMenu* GeneratedMenu)
 		}
 
 		FMenuBuilder MenuBuilder(GeneratedMenu->bShouldCloseWindowAfterMenuSelection, GeneratedMenu->Context.CommandList, GeneratedMenu->Context.GetAllExtenders(), GeneratedMenu->bCloseSelfOnly, GeneratedMenu->StyleSet, GeneratedMenu->bSearchable, GeneratedMenu->MenuName);
+		MenuBuilder.SetExtendersEnabled(GeneratedMenu->bExtendersEnabled);
 		MenuBuilder.GetMultiBox()->WeakToolMenu = GeneratedMenu;
 		PopulateMenuBuilder(MenuBuilder, GeneratedMenu);
 		if (GeneratedMenu->ModifyBlockWidgetAfterMake.IsBound())
@@ -1526,6 +1577,7 @@ TSharedRef<SWidget> UToolMenus::GenerateWidget(UToolMenu* GeneratedMenu)
 	else if (GeneratedMenu->MenuType == EMultiBoxType::Menu)
 	{
 		FMenuBuilder MenuBuilder(GeneratedMenu->bShouldCloseWindowAfterMenuSelection, GeneratedMenu->Context.CommandList, GeneratedMenu->Context.GetAllExtenders(), GeneratedMenu->bCloseSelfOnly, GeneratedMenu->StyleSet, GeneratedMenu->bSearchable, GeneratedMenu->MenuName);
+		MenuBuilder.SetExtendersEnabled(GeneratedMenu->bExtendersEnabled);
 		PopulateMenuBuilder(MenuBuilder, GeneratedMenu);
 		TSharedRef<SWidget> Result = MenuBuilder.MakeWidget(nullptr, GeneratedMenu->MaxHeight);
 		GeneratedMenuWidget.Widget = Result;
@@ -1534,6 +1586,7 @@ TSharedRef<SWidget> UToolMenus::GenerateWidget(UToolMenu* GeneratedMenu)
 	else if (GeneratedMenu->MenuType == EMultiBoxType::MenuBar)
 	{
 		FMenuBarBuilder MenuBarBuilder(GeneratedMenu->Context.CommandList, GeneratedMenu->Context.GetAllExtenders(), GeneratedMenu->StyleSet, GeneratedMenu->MenuName);
+		MenuBarBuilder.SetExtendersEnabled(GeneratedMenu->bExtendersEnabled);
 		PopulateMenuBarBuilder(MenuBarBuilder, GeneratedMenu);
 		TSharedRef<SWidget> Result = MenuBarBuilder.MakeWidget();
 		GeneratedMenuWidget.Widget = Result;
@@ -1543,6 +1596,7 @@ TSharedRef<SWidget> UToolMenus::GenerateWidget(UToolMenu* GeneratedMenu)
 	{
 		const EOrientation ToolBarOrientation = (GeneratedMenu->MenuType == EMultiBoxType::VerticalToolBar) ? Orient_Vertical : Orient_Horizontal;
 		FToolBarBuilder ToolbarBuilder(GeneratedMenu->Context.CommandList, GeneratedMenu->MenuName, GeneratedMenu->Context.GetAllExtenders(), ToolBarOrientation, GeneratedMenu->bToolBarForceSmallIcons);
+		ToolbarBuilder.SetExtendersEnabled(GeneratedMenu->bExtendersEnabled);
 		ToolbarBuilder.SetIsFocusable(GeneratedMenu->bToolBarIsFocusable);
 		PopulateToolBarBuilder(ToolbarBuilder, GeneratedMenu);
 		TSharedRef<SWidget> Result = ToolbarBuilder.MakeWidget();
diff --git a/Engine/Source/Developer/ToolMenus/Public/ToolMenu.h b/Engine/Source/Developer/ToolMenus/Public/ToolMenu.h
--- a/Engine/Source/Developer/ToolMenus/Public/ToolMenu.h
+++ b/Engine/Source/Developer/ToolMenus/Public/ToolMenu.h
@@ -62,6 +62,9 @@ public:
 	/** returns "SubMenuC.SubMenuD" for menu "ModuleA.MenuB.SubMenuC.SubMenuD" */
 	FString GetSubMenuNamePath() const;
 
+	/* Set support for extenders */
+	void SetExtendersEnabled(bool bEnabled);
+
 	//~ Begin UToolMenuBase Interface
 	virtual bool IsEditing() const override;
 	virtual FName GetSectionName(const FName InEntryName) const override;
@@ -164,6 +167,7 @@ private:
 
 	bool bRegistered;
 	bool bIsRegistering;
+	bool bExtendersEnabled;
 
 	const ISlateStyle* StyleSet;
 
diff --git a/Engine/Source/Developer/ToolMenus/Public/ToolMenus.h b/Engine/Source/Developer/ToolMenus/Public/ToolMenus.h
--- a/Engine/Source/Developer/ToolMenus/Public/ToolMenus.h
+++ b/Engine/Source/Developer/ToolMenus/Public/ToolMenus.h
@@ -263,6 +263,9 @@ public:
 	/* Substitute one menu for another during generate but not during find or extend */
 	void AddMenuSubstitutionDuringGenerate(const FName OriginalMenu, const FName NewMenu);
 
+	/* Remove substitute one menu for another during generate */
+	void RemoveSubstitutionDuringGenerate(const FName InMenu);
+
 	/** Displaying extension points is for debugging menus */
 	DECLARE_DELEGATE_RetVal(bool, FShouldDisplayExtensionPoints);
 	FShouldDisplayExtensionPoints ShouldDisplayExtensionPoints;
diff --git a/Engine/Source/Runtime/Slate/Private/Framework/MultiBox/MultiBoxBuilder.cpp b/Engine/Source/Runtime/Slate/Private/Framework/MultiBox/MultiBoxBuilder.cpp
--- a/Engine/Source/Runtime/Slate/Private/Framework/MultiBox/MultiBoxBuilder.cpp
+++ b/Engine/Source/Runtime/Slate/Private/Framework/MultiBox/MultiBoxBuilder.cpp
@@ -320,13 +320,16 @@ void FMenuBuilder::AddSearchWidget()
 
 void FMenuBuilder::ApplyHook(FName InExtensionHook, EExtensionHook::Position HookPosition)
 {
-	// this is a virtual function to get a properly typed "this" pointer
-	auto& Extender = ExtenderStack.Top();
-	if (InExtensionHook != NAME_None && Extender.IsValid())
+	if (ExtendersEnabled())
 	{
-		if (!MultiBox->IsInEditMode())
+		// this is a virtual function to get a properly typed "this" pointer
+		auto& Extender = ExtenderStack.Top();
+		if (InExtensionHook != NAME_None && Extender.IsValid())
 		{
-			Extender->Apply(InExtensionHook, HookPosition, *this);
+			if (!MultiBox->IsInEditMode())
+			{
+				Extender->Apply(InExtensionHook, HookPosition, *this);
+			}
 		}
 	}
 }
@@ -369,11 +372,14 @@ void FMenuBarBuilder::AddPullDownMenu(const FText& InMenuLabel, const FText& InT
 
 void FMenuBarBuilder::ApplyHook(FName InExtensionHook, EExtensionHook::Position HookPosition)
 {
-	// this is a virtual function to get a properly typed "this" pointer
-	auto& Extender = ExtenderStack.Top();
-	if (InExtensionHook != NAME_None && Extender.IsValid())
+	if (ExtendersEnabled())
 	{
-		Extender->Apply(InExtensionHook, HookPosition, *this);
+		// this is a virtual function to get a properly typed "this" pointer
+		auto& Extender = ExtenderStack.Top();
+		if (InExtensionHook != NAME_None && Extender.IsValid())
+		{
+			Extender->Apply(InExtensionHook, HookPosition, *this);
+		}
 	}
 }
 
@@ -540,11 +546,14 @@ void FToolBarBuilder::EndSection()
 
 void FToolBarBuilder::ApplyHook(FName InExtensionHook, EExtensionHook::Position HookPosition)
 {
-	// this is a virtual function to get a properly typed "this" pointer
-	auto& Extender = ExtenderStack.Top();
-	if (InExtensionHook != NAME_None && Extender.IsValid())
+	if (ExtendersEnabled())
 	{
-		Extender->Apply(InExtensionHook, HookPosition, *this);
+		// this is a virtual function to get a properly typed "this" pointer
+		auto& Extender = ExtenderStack.Top();
+		if (InExtensionHook != NAME_None && Extender.IsValid())
+		{
+			Extender->Apply(InExtensionHook, HookPosition, *this);
+		}
 	}
 }
 
diff --git a/Engine/Source/Runtime/Slate/Public/Framework/MultiBox/MultiBoxBuilder.h b/Engine/Source/Runtime/Slate/Public/Framework/MultiBox/MultiBoxBuilder.h
--- a/Engine/Source/Runtime/Slate/Public/Framework/MultiBox/MultiBoxBuilder.h
+++ b/Engine/Source/Runtime/Slate/Public/Framework/MultiBox/MultiBoxBuilder.h
@@ -119,6 +119,12 @@ public:
 	/** @return  The customization settings for the box being built */
 	FMultiBoxCustomization GetCustomization() const;
 
+	/** Sets extender support */
+	void SetExtendersEnabled(bool bEnabled) { bExtendersEnabled = bEnabled; }
+
+	/** @return True if extenders are enabled */
+	bool ExtendersEnabled() const { return bExtendersEnabled; }
+
 protected:
 	/** Applies any potential extension hooks at the current place */
 	virtual void ApplyHook(FName InExtensionHook, EExtensionHook::Position HookPosition) = 0;
@@ -143,6 +149,9 @@ protected:
 
 	/** Name of the menu */
 	FName MenuName;
+
+	/** If extenders are enabled */
+	bool bExtendersEnabled;
 };
 
 
diff --git a/Engine/Source/Runtime/Slate/Public/Framework/MultiBox/ToolMenuBase.h b/Engine/Source/Runtime/Slate/Public/Framework/MultiBox/ToolMenuBase.h
--- a/Engine/Source/Runtime/Slate/Public/Framework/MultiBox/ToolMenuBase.h
+++ b/Engine/Source/Runtime/Slate/Public/Framework/MultiBox/ToolMenuBase.h
@@ -73,7 +73,6 @@ struct SLATE_API FCustomizedToolMenu
 	UPROPERTY()
 	TArray<FName> SectionOrder;
 
-	UPROPERTY()
 	TArray<FName> WhitelistEntries;
 
 	FCustomizedToolMenuEntry* FindEntry(const FName InEntryName);
