Fix menu extenders showing up when disabled

diff --git a/Engine/Source/Developer/ToolMenus/Private/ToolMenus.cpp b/Engine/Source/Developer/ToolMenus/Private/ToolMenus.cpp
--- a/Engine/Source/Developer/ToolMenus/Private/ToolMenus.cpp
+++ b/Engine/Source/Developer/ToolMenus/Private/ToolMenus.cpp
@@ -755,6 +755,7 @@ void UToolMenus::PopulateSubMenu(FMenuBuilder& MenuBuilder, TWeakObjectPtr<UTool
 {
 	if (UToolMenu* GeneratedMenu = GenerateSubMenu(InParent.Get(), InBlockName))
 	{
+		MenuBuilder.SetExtendersEnabled(GeneratedMenu->bExtendersEnabled);
 		PopulateMenuBuilder(MenuBuilder, GeneratedMenu);
 	}
 }
@@ -769,21 +770,6 @@ TSharedRef<SWidget> UToolMenus::GenerateToolbarComboButtonMenu(TWeakObjectPtr<UT
 	return SNullWidget::NullWidget;
 }
 
-void UToolMenus::FillMenuBarDropDown(class FMenuBuilder& MenuBuilder, FName InParentName, FName InChildName, FToolMenuContext InMenuContext)
-{
-	if (UToolMenu* MenuToUse = FindSubMenuToGenerateWith(InParentName, InChildName))
-	{
-		// Create combined final menu
-		UToolMenu* GeneratedMenu = GenerateMenu(MenuToUse->MenuName, InMenuContext);
-
-		// Override name that people should use to extend this menu (as it may not be implemented)
-		GeneratedMenu->MenuName = JoinMenuPaths(InParentName, InChildName);
-
-		// Populate menu builder with final menu
-		PopulateMenuBuilder(MenuBuilder, GeneratedMenu);
-	}
-}
-
 void UToolMenus::PopulateMenuBuilder(FMenuBuilder& MenuBuilder, UToolMenu* MenuData)
 {
 	const bool bIsEditing = MenuData->IsEditing();
@@ -1741,12 +1727,14 @@ bool UToolMenus::RefreshMenuWidget(const FName InName, FGeneratedToolMenuWidget&
 	if (GeneratedMenu->MenuType == EMultiBoxType::Menu)
 	{
 		FMenuBuilder MenuBuilder(GeneratedMenu->bShouldCloseWindowAfterMenuSelection, GeneratedMenu->Context.CommandList, GeneratedMenu->Context.GetAllExtenders(), GeneratedMenu->bCloseSelfOnly, GeneratedMenu->StyleSet, GeneratedMenu->bSearchable);
+		MenuBuilder.SetExtendersEnabled(GeneratedMenu->bExtendersEnabled);
 		PopulateMenuBuilder(MenuBuilder, GeneratedMenu);
 		MultiBoxWidget->SetMultiBox(MenuBuilder.GetMultiBox());
 	}
 	else if (GeneratedMenu->MenuType == EMultiBoxType::MenuBar)
 	{
 		FMenuBarBuilder MenuBarBuilder(GeneratedMenu->Context.CommandList, GeneratedMenu->Context.GetAllExtenders(), GeneratedMenu->StyleSet);
+		MenuBarBuilder.SetExtendersEnabled(GeneratedMenu->bExtendersEnabled);
 		PopulateMenuBarBuilder(MenuBarBuilder, GeneratedMenu);
 		MultiBoxWidget->SetMultiBox(MenuBarBuilder.GetMultiBox());
 	}
@@ -1754,6 +1742,7 @@ bool UToolMenus::RefreshMenuWidget(const FName InName, FGeneratedToolMenuWidget&
 	{
 		const EOrientation ToolBarOrientation = (GeneratedMenu->MenuType == EMultiBoxType::VerticalToolBar) ? Orient_Vertical : Orient_Horizontal;
 		FToolBarBuilder ToolbarBuilder(GeneratedMenu->Context.CommandList, GeneratedMenu->MenuName, GeneratedMenu->Context.GetAllExtenders(), ToolBarOrientation, GeneratedMenu->bToolBarForceSmallIcons);
+		ToolbarBuilder.SetExtendersEnabled(GeneratedMenu->bExtendersEnabled);
 		ToolbarBuilder.SetIsFocusable(GeneratedMenu->bToolBarIsFocusable);
 		PopulateToolBarBuilder(ToolbarBuilder, GeneratedMenu);
 		MultiBoxWidget->SetMultiBox(ToolbarBuilder.GetMultiBox());
diff --git a/Engine/Source/Developer/ToolMenus/Public/ToolMenus.h b/Engine/Source/Developer/ToolMenus/Public/ToolMenus.h
--- a/Engine/Source/Developer/ToolMenus/Public/ToolMenus.h
+++ b/Engine/Source/Developer/ToolMenus/Public/ToolMenus.h
@@ -323,7 +323,6 @@ private:
 
 	UToolMenu* FindSubMenuToGenerateWith(const FName InParentName, const FName InChildName);
 
-	void FillMenuBarDropDown(FMenuBuilder& MenuBuilder, FName InParentName, FName InChildName, FToolMenuContext InMenuContext);
 	void PopulateMenuBuilder(FMenuBuilder& MenuBuilder, UToolMenu* MenuData);
 	void PopulateMenuBarBuilder(FMenuBarBuilder& MenuBarBuilder, UToolMenu* MenuData);
 	void PopulateToolBarBuilder(FToolBarBuilder& ToolBarBuilder, UToolMenu* MenuData);
