Fix support for submenus that do not have a name set
#rnx

diff --git a/Engine/Source/Developer/ToolMenus/Private/ToolMenus.cpp b/Engine/Source/Developer/ToolMenus/Private/ToolMenus.cpp
--- a/Engine/Source/Developer/ToolMenus/Private/ToolMenus.cpp
+++ b/Engine/Source/Developer/ToolMenus/Private/ToolMenus.cpp
@@ -760,6 +760,31 @@ void UToolMenus::PopulateSubMenu(FMenuBuilder& MenuBuilder, TWeakObjectPtr<UTool
 	}
 }
 
+void UToolMenus::PopulateSubMenuWithoutName(FMenuBuilder& MenuBuilder, TWeakObjectPtr<UToolMenu> InParent, const FNewToolMenuDelegate InNewToolMenuDelegate)
+{
+	const UToolMenu* InGeneratedParent = InParent.Get();
+	if (InGeneratedParent == nullptr)
+	{
+		return;
+	}
+
+	if (InNewToolMenuDelegate.IsBound())
+	{
+		UToolMenu* Menu = NewObject<UToolMenu>(this);
+		Menu->Context = InGeneratedParent->Context;
+		Menu->MenuName = NAME_None; // Menu does not have a name
+
+		// Submenu specific data
+		Menu->SubMenuParent = InGeneratedParent;
+		Menu->SubMenuSourceEntryName = NAME_None; // Entry does not have a name
+
+		InNewToolMenuDelegate.Execute(Menu);
+		Menu->MenuName = NAME_None; // Menu does not have a name
+
+		PopulateMenuBuilder(MenuBuilder, Menu);
+	}
+}
+
 TSharedRef<SWidget> UToolMenus::GenerateToolbarComboButtonMenu(TWeakObjectPtr<UToolMenu> InParent, const FName InBlockName)
 {
 	if (UToolMenu* GeneratedMenu = GenerateSubMenu(InParent.Get(), InBlockName))
@@ -855,6 +880,18 @@ void UToolMenus::PopulateMenuBuilder(FMenuBuilder& MenuBuilder, UToolMenu* MenuD
 					{
 						NewMenuDelegate = Block.SubMenuData.ConstructMenu.NewMenuDelegate;
 					}
+					else if (Block.Name == NAME_None)
+					{
+						if (Block.SubMenuData.ConstructMenu.NewToolMenuDelegate.IsBound())
+						{
+							// Blocks with no name cannot call PopulateSubMenu()
+							NewMenuDelegate = FNewMenuDelegate::CreateUObject(this, &UToolMenus::PopulateSubMenuWithoutName, TWeakObjectPtr<UToolMenu>(MenuData), Block.SubMenuData.ConstructMenu.NewToolMenuDelegate);
+						}
+						else
+						{
+							UE_LOG(LogToolMenus, Warning, TEXT("Submenu that has no name is missing required delegate: %s"), *MenuData->MenuName.ToString());
+						}
+					}
 					else
 					{
 						NewMenuDelegate = FNewMenuDelegate::CreateUObject(this, &UToolMenus::PopulateSubMenu, TWeakObjectPtr<UToolMenu>(MenuData), Block.Name);
diff --git a/Engine/Source/Developer/ToolMenus/Public/ToolMenus.h b/Engine/Source/Developer/ToolMenus/Public/ToolMenus.h
--- a/Engine/Source/Developer/ToolMenus/Public/ToolMenus.h
+++ b/Engine/Source/Developer/ToolMenus/Public/ToolMenus.h
@@ -341,6 +341,7 @@ private:
 	static void ExecuteStringCommand(const FToolMenuStringCommand StringCommand, const FToolMenuContext Context);
 
 	void PopulateSubMenu(FMenuBuilder& Builder, TWeakObjectPtr<UToolMenu> InParent, const FName InBlockName);
+	void PopulateSubMenuWithoutName(FMenuBuilder& MenuBuilder, TWeakObjectPtr<UToolMenu> InParent, const FNewToolMenuDelegate InNewToolMenuDelegate);
 
 	void ListAllParents(const FName Name, TArray<FName>& AllParents);
 
