Move ToolMenus module from Runtime to Developer, also removed one dependency on engine (Timer)
#rnx

diff --git a/Engine/Source/Developer/ToolMenus/Private/ToolMenu.cpp b/Engine/Source/Developer/ToolMenus/Private/ToolMenu.cpp
new file mode 100644
--- /dev/null
+++ b/Engine/Source/Developer/ToolMenus/Private/ToolMenu.cpp
@@ -0,0 +1,219 @@
+// Copyright 1998-2019 Epic Games, Inc. All Rights Reserved.
+
+#include "ToolMenu.h"
+#include "ToolMenus.h"
+#include "IToolMenusModule.h"
+
+#include "Textures/SlateIcon.h"
+#include "Framework/Commands/UIAction.h"
+#include "Framework/MultiBox/MultiBox.h"
+#include "Framework/MultiBox/MultiBoxBuilder.h"
+#include "Internationalization/Internationalization.h"
+
+UToolMenu::UToolMenu() :
+	MenuType(EMultiBoxType::Menu)
+	, bShouldCloseWindowAfterMenuSelection(true)
+	, bCloseSelfOnly(false)
+	, bSearchable(false)
+	, bToolBarIsFocusable(false)
+	, bToolBarForceSmallIcons(false)
+	, bRegistered(false)
+	, StyleSet(&FCoreStyle::Get())
+{
+}
+
+void UToolMenu::InitMenu(const FToolMenuOwner InOwner, FName InName, FName InParent, EMultiBoxType InType)
+{
+	MenuOwner = InOwner;
+	MenuName = InName;
+	MenuParent = InParent;
+	MenuType = InType;
+}
+
+void UToolMenu::InitGeneratedCopy(const UToolMenu* Source)
+{
+	// Skip sections and context
+
+	MenuName = Source->MenuName;
+	MenuParent = Source->MenuParent;
+	StyleName = Source->StyleName;
+	TutorialHighlightName = Source->TutorialHighlightName;
+	MenuType = Source->MenuType;
+	StyleSet = Source->StyleSet;
+	bShouldCloseWindowAfterMenuSelection = Source->bShouldCloseWindowAfterMenuSelection;
+	bCloseSelfOnly = Source->bCloseSelfOnly;
+	bSearchable = Source->bSearchable;
+	bToolBarIsFocusable = Source->bToolBarIsFocusable;
+	bToolBarForceSmallIcons = Source->bToolBarForceSmallIcons;
+	MenuOwner = Source->MenuOwner;
+}
+
+int32 UToolMenu::IndexOfSection(const FName InSectionName) const
+{
+	for (int32 i=0; i < Sections.Num(); ++i)
+	{
+		if (Sections[i].Name == InSectionName)
+		{
+			return i;
+		}
+	}
+
+	return INDEX_NONE;
+}
+
+int32 UToolMenu::FindInsertIndex(const FToolMenuSection& InSection) const
+{
+	const FToolMenuInsert InInsertPosition = InSection.InsertPosition;
+	if (InInsertPosition.IsDefault())
+	{
+		return Sections.Num();
+	}
+
+	if (InInsertPosition.Position == EToolMenuInsertType::First)
+	{
+		for (int32 i = 0; i < Sections.Num(); ++i)
+		{
+			if (Sections[i].InsertPosition.Position != InInsertPosition.Position)
+			{
+				return i;
+			}
+		}
+
+		return Sections.Num();
+	}
+
+	int32 DestIndex = IndexOfSection(InInsertPosition.Name);
+	if (DestIndex == INDEX_NONE)
+	{
+		return DestIndex;
+	}
+
+	if (InInsertPosition.Position == EToolMenuInsertType::After)
+	{
+		++DestIndex;
+	}
+
+	for (int32 i = DestIndex; i < Sections.Num(); ++i)
+	{
+		if (Sections[i].InsertPosition != InInsertPosition)
+		{
+			return i;
+		}
+	}
+
+	return Sections.Num();
+}
+
+FToolMenuSection& UToolMenu::AddDynamicSection(const FName SectionName, const FNewSectionConstructChoice& InConstruct, const FToolMenuInsert InPosition)
+{
+	FToolMenuSection& Section = AddSection(SectionName, TAttribute< FText >(), InPosition);
+	Section.Construct = InConstruct;
+	return Section;
+}
+
+FToolMenuSection& UToolMenu::AddSection(const FName SectionName, const TAttribute< FText >& InLabel, const FToolMenuInsert InPosition)
+{
+	for (FToolMenuSection& Section : Sections)
+	{
+		if (Section.Name == SectionName)
+		{
+			if (InLabel.IsSet())
+			{
+				Section.Label = InLabel;
+			}
+
+			if (InPosition.Name != NAME_None)
+			{
+				Section.InsertPosition = InPosition;
+			}
+
+			return Section;
+		}
+	}
+
+	FToolMenuSection& NewSection = Sections.AddDefaulted_GetRef();
+	NewSection.InitSection(SectionName, InLabel, InPosition);
+	return NewSection;
+}
+
+void UToolMenu::AddSectionScript(const FName SectionName, const FText& InLabel, const FName InsertName, const EToolMenuInsertType InsertType)
+{
+	FToolMenuSection& Section = FindOrAddSection(SectionName);
+	Section.Label = InLabel;
+	Section.InsertPosition = FToolMenuInsert(InsertName, InsertType);
+}
+
+void UToolMenu::AddDynamicSectionScript(const FName SectionName, UToolMenuSectionDynamic* InObject)
+{
+	FToolMenuSection& Section = FindOrAddSection(SectionName);
+	Section.ToolMenuSectionDynamic = InObject;
+}
+
+void UToolMenu::AddMenuEntryObject(UToolMenuEntryScript* InObject)
+{
+	FindOrAddSection(InObject->Data.Section).AddEntryObject(InObject);
+}
+
+UToolMenu* UToolMenu::AddSubMenuScript(const FName InOwner, const FName SectionName, const FName InName, const FText& InLabel, const FText& InToolTip)
+{
+	return AddSubMenu(InOwner, SectionName, InName, InLabel, InToolTip);
+}
+
+UToolMenu* UToolMenu::AddSubMenu(const FToolMenuOwner InOwner, const FName SectionName, const FName InName, const FText& InLabel, const FText& InToolTip)
+{
+	FToolMenuEntry Args = FToolMenuEntry::InitSubMenu(MenuName, InName, InLabel, InToolTip, FNewToolMenuChoice());
+	Args.Owner = InOwner;
+	FindOrAddSection(SectionName).AddEntry(Args);
+	return UToolMenus::Get()->ExtendMenu(*(MenuName.ToString() + TEXT(".") + InName.ToString()));
+}
+
+FToolMenuSection* UToolMenu::FindSection(const FName SectionName)
+{
+	for (FToolMenuSection& Section : Sections)
+	{
+		if (Section.Name == SectionName)
+		{
+			return &Section;
+		}
+	}
+
+	return nullptr;
+}
+
+FToolMenuSection& UToolMenu::FindOrAddSection(const FName SectionName)
+{
+	for (FToolMenuSection& Section : Sections)
+	{
+		if (Section.Name == SectionName)
+		{
+			return Section;
+		}
+	}
+	
+	return AddSection(SectionName);
+}
+
+void UToolMenu::RemoveSection(const FName SectionName)
+{
+	Sections.RemoveAll([SectionName](const FToolMenuSection& Section) { return Section.Name == SectionName; });
+}
+
+bool UToolMenu::FindEntry(const FName EntryName, int32& SectionIndex, int32& EntryIndex) const
+{
+	for (int32 i=0; i < Sections.Num(); ++i)
+	{
+		EntryIndex = Sections[i].IndexOfBlock(EntryName);
+		if (EntryIndex != INDEX_NONE)
+		{
+			SectionIndex = i;
+			return true;
+		}
+	}
+
+	return false;
+}
+
+void UToolMenu::AddMenuEntry(const FName SectionName, const FToolMenuEntry& Args)
+{
+	FindOrAddSection(SectionName).AddEntry(Args);
+}
diff --git a/Engine/Source/Developer/ToolMenus/Private/ToolMenuContext.cpp b/Engine/Source/Developer/ToolMenus/Private/ToolMenuContext.cpp
new file mode 100644
--- /dev/null
+++ b/Engine/Source/Developer/ToolMenus/Private/ToolMenuContext.cpp
@@ -0,0 +1,127 @@
+// Copyright 1998-2019 Epic Games, Inc. All Rights Reserved.
+
+#include "ToolMenuContext.h"
+#include "ToolMenus.h"
+#include "IToolMenusModule.h"
+
+#include "Textures/SlateIcon.h"
+#include "Framework/Commands/UIAction.h"
+#include "Framework/MultiBox/MultiBox.h"
+#include "Framework/MultiBox/MultiBoxBuilder.h"
+#include "Internationalization/Internationalization.h"
+
+
+FToolMenuContext::FToolMenuContext()
+{
+
+}
+
+FToolMenuContext::FToolMenuContext(UObject* InContext)
+{
+	if (InContext)
+	{
+		ContextObjects.Add(InContext);
+	}
+}
+
+FToolMenuContext::FToolMenuContext(TSharedPtr<FUICommandList> InCommandList, TSharedPtr<FExtender> InExtender, UObject* InContext)
+{
+	if (InContext)
+	{
+		ContextObjects.Add(InContext);
+	}
+
+	if (InExtender.IsValid())
+	{
+		AddExtender(InExtender);
+	}
+
+	AppendCommandList(InCommandList);
+}
+
+UObject* FToolMenuContext::FindByClass(UClass* InClass) const
+{
+	for (UObject* ContextObject : ContextObjects)
+	{
+		if (ContextObject && ContextObject->IsA(InClass))
+		{
+			return ContextObject;
+		}
+	}
+
+	return nullptr;
+}
+
+void FToolMenuContext::AppendCommandList(const TSharedRef<FUICommandList>& InCommandList)
+{
+	const TSharedPtr<FUICommandList> List = InCommandList;
+	AppendCommandList(List);
+}
+
+void FToolMenuContext::AppendCommandList(const TSharedPtr<FUICommandList>& InCommandList)
+{
+	if (InCommandList.IsValid())
+	{
+		CommandLists.Add(InCommandList);
+
+		if (CommandLists.Num() == 1)
+		{
+			CommandList = InCommandList;
+		}
+		else if (CommandLists.Num() == 2)
+		{
+			CommandList = MakeShared<FUICommandList>();
+			CommandList->Append(CommandLists[0].ToSharedRef());
+			CommandList->Append(InCommandList.ToSharedRef());
+		}
+		else
+		{
+			CommandList->Append(InCommandList.ToSharedRef());
+		}
+	}
+}
+
+const FUIAction* FToolMenuContext::GetActionForCommand(TSharedPtr<const FUICommandInfo> Command, TSharedPtr<const FUICommandList>& OutCommandList) const
+{
+	for (const TSharedPtr<FUICommandList>& CommandListIter : CommandLists)
+	{
+		if (CommandListIter.IsValid())
+		{
+			if (const FUIAction* Result = CommandListIter->GetActionForCommand(Command))
+			{
+				OutCommandList = CommandListIter;
+				return Result;
+			}
+		}
+	}
+
+	return nullptr;
+}
+
+void FToolMenuContext::AddExtender(const TSharedPtr<FExtender>& InExtender)
+{
+	Extenders.AddUnique(InExtender);
+}
+
+TSharedPtr<FExtender> FToolMenuContext::GetAllExtenders()
+{
+	return FExtender::Combine(Extenders);
+}
+
+void FToolMenuContext::ResetExtenders()
+{
+	Extenders.Reset();
+}
+
+void FToolMenuContext::AppendObjects(const TArray<UObject*>& InObjects)
+{
+	for (UObject* Object : InObjects)
+	{
+		AddObject(Object);
+	}
+}
+
+void FToolMenuContext::AddObject(UObject* InObject)
+{
+	ContextObjects.AddUnique(InObject);
+}
diff --git a/Engine/Source/Developer/ToolMenus/Private/ToolMenuEntry.cpp b/Engine/Source/Developer/ToolMenus/Private/ToolMenuEntry.cpp
new file mode 100644
--- /dev/null
+++ b/Engine/Source/Developer/ToolMenus/Private/ToolMenuEntry.cpp
@@ -0,0 +1,218 @@
+// Copyright 1998-2019 Epic Games, Inc. All Rights Reserved.
+
+#include "ToolMenuEntry.h"
+#include "ToolMenus.h"
+#include "IToolMenusModule.h"
+
+#include "Textures/SlateIcon.h"
+#include "Framework/Commands/UIAction.h"
+#include "Framework/MultiBox/MultiBox.h"
+#include "Framework/MultiBox/MultiBoxBuilder.h"
+#include "Internationalization/Internationalization.h"
+
+
+FToolMenuEntry::FToolMenuEntry() :
+	Type(EMultiBlockType::None),
+	UserInterfaceActionType(EUserInterfaceActionType::Button),
+	bShouldCloseWindowAfterMenuSelection(true),
+	ScriptObject(nullptr)
+{
+}
+
+FToolMenuEntry::FToolMenuEntry(const FToolMenuOwner InOwner, const FName InName, EMultiBlockType InType) :
+	Name(InName),
+	Owner(InOwner),
+	Type(InType),
+	UserInterfaceActionType(EUserInterfaceActionType::Button),
+	bShouldCloseWindowAfterMenuSelection(true),
+	ScriptObject(nullptr)
+{
+}
+
+const FUIAction* FToolMenuEntry::GetActionForCommand(const FToolMenuContext& InContext, TSharedPtr<const FUICommandList>& OutCommandList) const
+{
+	if (Command.IsValid())
+	{
+		if (CommandList.IsValid())
+		{
+			const FUIAction* Result = CommandList->GetActionForCommand(Command);
+			if (Result)
+			{
+				OutCommandList = CommandList;
+				return Result;
+			}
+		}
+		else
+		{
+			return InContext.GetActionForCommand(Command, OutCommandList);
+		}
+	}
+
+	return nullptr;
+}
+
+void FToolMenuEntry::SetCommandList(const TSharedPtr<const FUICommandList>& InCommandList)
+{
+	CommandList = InCommandList;
+}
+
+void FToolMenuEntry::SetCommand(const TSharedPtr<const FUICommandInfo>& InCommand, FName InName, const TAttribute<FText>& InLabel, const TAttribute<FText>& InToolTip, const TAttribute<FSlateIcon>& InIcon)
+{
+	Command = InCommand;
+	Name = InName != NAME_None ? InName : InCommand->GetCommandName();
+	Label = InLabel.IsSet() ? InLabel : InCommand->GetLabel();
+	ToolTip = InToolTip.IsSet() ? InToolTip : InCommand->GetDescription();
+	Icon = InIcon.IsSet() ? InIcon : InCommand->GetIcon();
+}
+
+FToolMenuEntry FToolMenuEntry::InitMenuEntry(const FName InName, const TAttribute<FText>& InLabel, const TAttribute<FText>& InToolTip, const TAttribute<FSlateIcon>& InIcon, const FToolUIActionChoice& InAction, const EUserInterfaceActionType InUserInterfaceActionType, const FName InTutorialHighlightName)
+{
+	FToolMenuEntry Entry(UToolMenus::Get()->CurrentOwner(), InName, EMultiBlockType::MenuEntry);
+	Entry.TutorialHighlightName = InTutorialHighlightName;
+	Entry.Label = InLabel;
+	Entry.ToolTip = InToolTip;
+	Entry.Icon = InIcon;
+	Entry.UserInterfaceActionType = InUserInterfaceActionType;
+	Entry.Action = InAction;
+	return Entry;
+}
+
+FToolMenuEntry FToolMenuEntry::InitMenuEntry(const TSharedPtr< const FUICommandInfo >& InCommand, const TAttribute<FText>& InLabel, const TAttribute<FText>& InToolTip, const TAttribute<FSlateIcon>& InIcon, const FName InTutorialHighlightName, const FName InName)
+{
+	FToolMenuEntry Entry(UToolMenus::Get()->CurrentOwner(), InName, EMultiBlockType::MenuEntry);
+	Entry.TutorialHighlightName = InTutorialHighlightName;
+	Entry.SetCommand(InCommand, InName, InLabel, InToolTip, InIcon);
+	Entry.CommandList.Reset();
+	return Entry;
+}
+
+FToolMenuEntry FToolMenuEntry::InitMenuEntryWithCommandList(const TSharedPtr< const FUICommandInfo >& InCommand, const TSharedPtr< const FUICommandList >& InCommandList, const TAttribute<FText>& InLabel, const TAttribute<FText>& InToolTip, const TAttribute<FSlateIcon>& InIcon, const FName InTutorialHighlightName, const FName InName)
+{
+	FToolMenuEntry Entry(UToolMenus::Get()->CurrentOwner(), InName, EMultiBlockType::MenuEntry);
+	Entry.TutorialHighlightName = InTutorialHighlightName;
+	Entry.SetCommand(InCommand, InName, InLabel, InToolTip, InIcon);
+	Entry.CommandList = InCommandList;
+	return Entry;
+}
+
+FToolMenuEntry FToolMenuEntry::InitMenuEntry(const FName InName, const FToolUIActionChoice& InAction, const TSharedRef<SWidget>& Widget)
+{
+	FToolMenuEntry Entry(UToolMenus::Get()->CurrentOwner(), InName, EMultiBlockType::MenuEntry);
+	Entry.Action = InAction;
+	Entry.MakeWidget.BindLambda([Widget](const FToolMenuContext&) { return Widget; });
+	return Entry;
+}
+
+FToolMenuEntry FToolMenuEntry::InitSubMenu(const FName InParentMenu, const FName InName, const TAttribute<FText>& InLabel, const TAttribute<FText>& InToolTip, const FNewToolMenuChoice& InMakeMenu, const FToolUIActionChoice& InAction, const EUserInterfaceActionType InUserInterfaceActionType, bool bInOpenSubMenuOnClick, const TAttribute<FSlateIcon>& InIcon, const bool bInShouldCloseWindowAfterMenuSelection)
+{
+	FToolMenuEntry Entry(UToolMenus::Get()->CurrentOwner(), InName, EMultiBlockType::MenuEntry);
+	Entry.Label = InLabel;
+	Entry.ToolTip = InToolTip;
+	Entry.Icon = InIcon;	
+	Entry.Action = InAction;
+	Entry.UserInterfaceActionType = InUserInterfaceActionType;
+	Entry.bShouldCloseWindowAfterMenuSelection = bInShouldCloseWindowAfterMenuSelection;
+	Entry.SubMenuData.bIsSubMenu = true;
+	Entry.SubMenuData.ConstructMenu = InMakeMenu;
+	Entry.SubMenuData.bOpenSubMenuOnClick = bInOpenSubMenuOnClick;
+	return Entry;
+}
+
+FToolMenuEntry FToolMenuEntry::InitSubMenu(const FName InParentMenu, const FName InName, const TAttribute<FText>& InLabel, const TAttribute<FText>& InToolTip, const FNewToolMenuChoice& InMakeMenu, bool bInOpenSubMenuOnClick, const TAttribute<FSlateIcon>& InIcon, const bool bInShouldCloseWindowAfterMenuSelection)
+{
+	FToolMenuEntry Entry(UToolMenus::Get()->CurrentOwner(), InName, EMultiBlockType::MenuEntry);
+	Entry.Label = InLabel;
+	Entry.ToolTip = InToolTip;
+	Entry.Icon = InIcon;
+	Entry.bShouldCloseWindowAfterMenuSelection = bInShouldCloseWindowAfterMenuSelection;
+	Entry.SubMenuData.bIsSubMenu = true;
+	Entry.SubMenuData.ConstructMenu = InMakeMenu;
+	Entry.SubMenuData.bOpenSubMenuOnClick = bInOpenSubMenuOnClick;
+	return Entry;
+}
+
+FToolMenuEntry FToolMenuEntry::InitSubMenu(const FName InParentMenu, const FName InName, const FToolUIActionChoice& InAction, const TSharedRef<SWidget>& InWidget, const FNewToolMenuChoice& InMakeMenu, bool bInShouldCloseWindowAfterMenuSelection)
+{
+	FToolMenuEntry Entry(UToolMenus::Get()->CurrentOwner(), InName, EMultiBlockType::MenuEntry);
+	Entry.Action = InAction;
+	Entry.MakeWidget.BindLambda([=](const FToolMenuContext&) { return InWidget; });
+	Entry.bShouldCloseWindowAfterMenuSelection = bInShouldCloseWindowAfterMenuSelection;
+	Entry.SubMenuData.bIsSubMenu = true;
+	Entry.SubMenuData.ConstructMenu = InMakeMenu;
+	Entry.SubMenuData.bOpenSubMenuOnClick = false;
+	return Entry;
+}
+
+FToolMenuEntry FToolMenuEntry::InitToolBarButton(const FName InName, const FToolUIActionChoice& InAction, const TAttribute<FText>& InLabel, const TAttribute<FText>& InToolTip, const TAttribute<FSlateIcon>& InIcon, const EUserInterfaceActionType InUserInterfaceActionType, FName InTutorialHighlightName)
+{
+	FToolMenuEntry Entry(UToolMenus::Get()->CurrentOwner(), InName, EMultiBlockType::ToolBarButton);
+	Entry.TutorialHighlightName = InTutorialHighlightName;
+	Entry.Label = InLabel;
+	Entry.ToolTip = InToolTip;
+	Entry.Icon = InIcon;
+	Entry.UserInterfaceActionType = InUserInterfaceActionType;
+	Entry.Action = InAction;
+	return Entry;
+}
+
+FToolMenuEntry FToolMenuEntry::InitToolBarButton(const TSharedPtr< const FUICommandInfo >& InCommand, const TAttribute<FText>& InLabel, const TAttribute<FText>& InToolTip, const TAttribute<FSlateIcon>& InIcon, FName InTutorialHighlightName, const FName InName)
+{
+	FToolMenuEntry Entry(UToolMenus::Get()->CurrentOwner(), InName, EMultiBlockType::ToolBarButton);
+	Entry.TutorialHighlightName = InTutorialHighlightName;
+	Entry.SetCommand(InCommand, InName, InLabel, InToolTip, InIcon);
+	return Entry;
+}
+
+FToolMenuEntry FToolMenuEntry::InitComboButton(const FName InName, const FToolUIActionChoice& InAction, const FNewToolMenuWidgetChoice& InMenuContentGenerator, const TAttribute<FText>& InLabel, const TAttribute<FText>& InToolTip, const TAttribute<FSlateIcon>& InIcon, bool bInSimpleComboBox, FName InTutorialHighlightName)
+{
+	FToolMenuEntry Entry(UToolMenus::Get()->CurrentOwner(), InName, EMultiBlockType::ToolBarComboButton);
+	Entry.TutorialHighlightName = InTutorialHighlightName;
+	Entry.Label = InLabel;
+	Entry.ToolTip = InToolTip;
+	Entry.Icon = InIcon;
+	Entry.Action = InAction;
+	Entry.ToolBarData.ComboButtonContextMenuGenerator = InMenuContentGenerator;
+	Entry.ToolBarData.bSimpleComboBox = bInSimpleComboBox;
+	return Entry;
+}
+
+FToolMenuEntry FToolMenuEntry::InitMenuSeparator(const FName InName)
+{
+	return FToolMenuEntry(UToolMenus::Get()->CurrentOwner(), InName, EMultiBlockType::MenuSeparator);
+}
+
+FToolMenuEntry FToolMenuEntry::InitToolBarSeparator(const FName InName)
+{
+	return FToolMenuEntry(UToolMenus::Get()->CurrentOwner(), InName, EMultiBlockType::ToolBarSeparator);
+}
+
+FToolMenuEntry FToolMenuEntry::InitWidget(const FName InName, const TSharedRef<SWidget>& InWidget, const FText& Label, bool bNoIndent, bool bSearchable)
+{
+	FToolMenuEntry Entry(UToolMenus::Get()->CurrentOwner(), InName, EMultiBlockType::Widget);
+	Entry.Label = Label;
+	Entry.MakeWidget.BindLambda([=](const FToolMenuContext&) { return InWidget; });
+	Entry.WidgetData.bNoIndent = bNoIndent;
+	Entry.WidgetData.bSearchable = bSearchable;
+	return Entry;
+}
+
+void FToolMenuEntry::ResetActions()
+{
+	Action = FToolUIActionChoice();
+	Command.Reset();
+	CommandList.Reset();
+	StringExecuteAction = FToolMenuStringCommand();
+	// Note: Cannot reset ScriptObject as it would also remove label and other data
+	//ScriptObject = nullptr;
+}
+
+bool FToolMenuEntry::IsNonLegacyDynamicConstruct() const
+{
+	return Construct.IsBound() || IsScriptObjectDynamicConstruct();
+}
+
+bool FToolMenuEntry::IsScriptObjectDynamicConstruct() const
+{
+	static const FName ConstructMenuEntryName = GET_FUNCTION_NAME_CHECKED(UToolMenuEntryScript, ConstructMenuEntry);
+	return ScriptObject && ScriptObject->GetClass()->IsFunctionImplementedInScript(ConstructMenuEntryName);
+}
diff --git a/Engine/Source/Developer/ToolMenus/Private/ToolMenuEntryScript.cpp b/Engine/Source/Developer/ToolMenus/Private/ToolMenuEntryScript.cpp
new file mode 100644
--- /dev/null
+++ b/Engine/Source/Developer/ToolMenus/Private/ToolMenuEntryScript.cpp
@@ -0,0 +1,170 @@
+// Copyright 1998-2019 Epic Games, Inc. All Rights Reserved.
+
+#include "ToolMenuEntryScript.h"
+#include "ToolMenuEntry.h"
+#include "ToolMenus.h"
+#include "IToolMenusModule.h"
+
+#include "Textures/SlateIcon.h"
+#include "Framework/Commands/UIAction.h"
+#include "Framework/MultiBox/MultiBox.h"
+#include "Framework/MultiBox/MultiBoxBuilder.h"
+#include "Internationalization/Internationalization.h"
+
+FScriptSlateIcon::FScriptSlateIcon()
+{
+
+}
+
+FScriptSlateIcon::FScriptSlateIcon(const FName InStyleSetName, const FName InStyleName) :
+	StyleSetName(InStyleSetName),
+	StyleName(InStyleName),
+	SmallStyleName(ISlateStyle::Join(InStyleName, ".Small"))
+{
+}
+
+FScriptSlateIcon::FScriptSlateIcon(const FName InStyleSetName, const FName InStyleName, const FName InSmallStyleName) :
+	StyleSetName(InStyleSetName),
+	StyleName(InStyleName),
+	SmallStyleName(InSmallStyleName)
+{
+}
+
+FSlateIcon FScriptSlateIcon::GetSlateIcon() const
+{
+	if (SmallStyleName == NAME_None)
+	{
+		if (StyleSetName == NAME_None && StyleName == NAME_None)
+		{
+			return FSlateIcon();
+		}
+
+		return FSlateIcon(StyleSetName, StyleName);
+	}
+
+	return FSlateIcon(StyleSetName, StyleName, SmallStyleName);
+}
+
+TAttribute<FText> UToolMenuEntryScript::CreateLabelAttribute(FToolMenuContext& Context)
+{
+	static const FName FunctionName = GET_FUNCTION_NAME_CHECKED(UToolMenuEntryScript, GetLabel);
+	if (GetClass()->IsFunctionImplementedInScript(FunctionName))
+	{
+		return TAttribute<FText>::Create(TAttribute<FText>::FGetter::CreateUFunction(this, FunctionName, Context));
+	}
+
+	return Data.Label;
+}
+
+TAttribute<FText> UToolMenuEntryScript::CreateToolTipAttribute(FToolMenuContext& Context)
+{
+	static const FName FunctionName = GET_FUNCTION_NAME_CHECKED(UToolMenuEntryScript, GetToolTip);
+	if (GetClass()->IsFunctionImplementedInScript(FunctionName))
+	{
+		return TAttribute<FText>::Create(TAttribute<FText>::FGetter::CreateUFunction(this, FunctionName, Context));
+	}
+
+	return Data.ToolTip;
+}
+
+TAttribute<FSlateIcon> UToolMenuEntryScript::CreateIconAttribute(FToolMenuContext& Context)
+{
+	static const FName FunctionName = GET_FUNCTION_NAME_CHECKED(UToolMenuEntryScript, GetIcon);
+	if (GetClass()->IsFunctionImplementedInScript(FunctionName))
+	{
+		TWeakObjectPtr<UToolMenuEntryScript> WeakThis(this);
+		TAttribute<FSlateIcon>::FGetter Getter;
+		Getter.BindLambda([=]()
+		{
+			if (UToolMenuEntryScript* Object = WeakThis.Get())
+			{
+				return Object->GetIcon(Context).GetSlateIcon();
+			}
+			else
+			{
+				return FSlateIcon();
+			}
+		});
+
+		return TAttribute<FSlateIcon>::Create(Getter);
+	}
+
+	return Data.Icon.GetSlateIcon();
+}
+
+FSlateIcon UToolMenuEntryScript::GetSlateIcon(const FToolMenuContext& Context) const
+{
+	return GetIcon(Context).GetSlateIcon();
+}
+
+void UToolMenuEntryScript::RegisterMenuEntry()
+{
+	UToolMenus::AddMenuEntryObject(this);
+}
+
+void UToolMenuEntryScript::InitEntry(const FName OwnerName, const FName Menu, const FName Section, const FName Name, const FText& Label, const FText& ToolTip)
+{
+	Data.OwnerName = OwnerName;
+	Data.Menu = Menu;
+	Data.Section = Section;
+	Data.Name = Name;
+	Data.Label = Label;
+	Data.ToolTip = ToolTip;
+}
+
+void UToolMenuEntryScript::ToMenuEntry(FToolMenuEntry& Output)
+{
+	if (Data.Advanced.bIsSubMenu)
+	{
+		Output = FToolMenuEntry::InitSubMenu(
+			Data.Menu,
+			Data.Name,
+			Data.Label,
+			Data.ToolTip,
+			FNewToolMenuChoice(), // Menu will be opened by string: 'Menu' + '.' + 'Name'
+			Data.Advanced.bOpenSubMenuOnClick,
+			Data.Icon,
+			Data.Advanced.bShouldCloseWindowAfterMenuSelection);
+	}
+	else
+	{
+		if (Data.Advanced.EntryType == EMultiBlockType::ToolBarButton)
+		{
+			Output = FToolMenuEntry::InitToolBarButton(
+				Data.Name,
+				FToolUIActionChoice(), // Action will be handled by 'ScriptObject'
+				Data.Label,
+				Data.ToolTip,
+				Data.Icon,
+				Data.Advanced.UserInterfaceActionType,
+				Data.Advanced.TutorialHighlight
+			);
+		}
+		else
+		{
+			Output = FToolMenuEntry::InitMenuEntry(Data.Name, Data.Label, Data.ToolTip, Data.Icon, FUIAction());
+			Output.UserInterfaceActionType = Data.Advanced.UserInterfaceActionType;
+			Output.TutorialHighlightName = Data.Advanced.TutorialHighlight;
+		}
+	}
+
+	if (!Data.InsertPosition.IsDefault())
+	{
+		Output.InsertPosition = Data.InsertPosition;
+	}
+
+	Output.ScriptObject = this;
+
+	Output.Owner = Data.OwnerName;
+}
+
+FToolMenuEntryScriptDataAdvanced::FToolMenuEntryScriptDataAdvanced() :
+	EntryType(EMultiBlockType::MenuEntry),
+	UserInterfaceActionType(EUserInterfaceActionType::Button),
+	bIsSubMenu(false),
+	bOpenSubMenuOnClick(false),
+	bShouldCloseWindowAfterMenuSelection(true),
+	bSimpleComboBox(false)
+{
+
+}
diff --git a/Engine/Source/Developer/ToolMenus/Private/ToolMenuSection.cpp b/Engine/Source/Developer/ToolMenus/Private/ToolMenuSection.cpp
new file mode 100644
--- /dev/null
+++ b/Engine/Source/Developer/ToolMenus/Private/ToolMenuSection.cpp
@@ -0,0 +1,213 @@
+// Copyright 1998-2019 Epic Games, Inc. All Rights Reserved.
+
+#include "ToolMenuSection.h"
+#include "ToolMenus.h"
+#include "IToolMenusModule.h"
+
+#include "Textures/SlateIcon.h"
+#include "Framework/Commands/UIAction.h"
+#include "Framework/MultiBox/MultiBox.h"
+#include "Framework/MultiBox/MultiBoxBuilder.h"
+#include "Internationalization/Internationalization.h"
+
+FToolMenuSection::FToolMenuSection() :
+	ToolMenuSectionDynamic(nullptr)
+{
+
+}
+
+
+void FToolMenuSection::InitSection(const FName InName, const TAttribute< FText >& InLabel, const FToolMenuInsert InPosition)
+{
+	Name = InName;
+	Label = InLabel;
+	InsertPosition = InPosition;
+}
+
+void FToolMenuSection::InitGeneratedSectionCopy(const FToolMenuSection& Source, FToolMenuContext& InContext)
+{
+	Name = Source.Name;
+	Label = Source.Label;
+	InsertPosition = Source.InsertPosition;
+	Construct = Source.Construct;
+	Context = InContext;
+}
+
+FToolMenuEntry& FToolMenuSection::AddEntry(const FToolMenuEntry& Args)
+{
+	if (Args.Name == NAME_None)
+	{
+		return Blocks.Add_GetRef(Args);
+	}
+
+	int32 BlockIndex = IndexOfBlock(Args.Name);
+	if (BlockIndex != INDEX_NONE)
+	{
+		Blocks[BlockIndex] = Args;
+		return Blocks[BlockIndex];
+	}
+	else
+	{
+		return Blocks.Add_GetRef(Args);
+	}
+}
+
+FToolMenuEntry& FToolMenuSection::AddEntryObject(UToolMenuEntryScript* InObject)
+{
+	// Avoid modifying objects that are saved as content on disk
+	UToolMenuEntryScript* DestObject = InObject;
+	if (DestObject->IsAsset())
+	{
+		DestObject = DuplicateObject<UToolMenuEntryScript>(InObject, UToolMenus::Get());
+	}
+
+	// Refresh widgets next tick so that toolbars and menu bars are updated
+	UToolMenus::Get()->RefreshAllWidgets();
+
+	FToolMenuEntry Args;
+	DestObject->ToMenuEntry(Args);
+	return AddEntry(Args);
+}
+
+FToolMenuEntry& FToolMenuSection::AddMenuEntry(const FName InName, const TAttribute<FText>& InLabel, const TAttribute<FText>& InToolTip, const TAttribute<FSlateIcon>& InIcon, const FToolUIActionChoice& InAction, const EUserInterfaceActionType InUserInterfaceActionType, const FName InTutorialHighlightName)
+{
+	return AddEntry(FToolMenuEntry::InitMenuEntry(InName, InLabel, InToolTip, InIcon, InAction, InUserInterfaceActionType, InTutorialHighlightName));
+}
+
+FToolMenuEntry& FToolMenuSection::AddMenuEntry(const TSharedPtr< const FUICommandInfo >& InCommand, const TAttribute<FText>& InLabelOverride, const TAttribute<FText>& InToolTipOverride, const TAttribute<FSlateIcon>& InIconOverride, const FName InTutorialHighlightName, const FName InNameOverride)
+{
+	return AddEntry(FToolMenuEntry::InitMenuEntry(InCommand, InLabelOverride, InToolTipOverride, InIconOverride, InTutorialHighlightName, InNameOverride));
+}
+
+FToolMenuEntry& FToolMenuSection::AddMenuEntryWithCommandList(const TSharedPtr< const FUICommandInfo >& InCommand, const TSharedPtr< const FUICommandList >& InCommandList, const TAttribute<FText>& InLabelOverride, const TAttribute<FText>& InToolTipOverride, const TAttribute<FSlateIcon>& InIconOverride, const FName InTutorialHighlightName, const FName InNameOverride)
+{
+	return AddEntry(FToolMenuEntry::InitMenuEntryWithCommandList(InCommand, InCommandList, InLabelOverride, InToolTipOverride, InIconOverride, InTutorialHighlightName, InNameOverride));
+}
+
+FToolMenuEntry& FToolMenuSection::AddDynamicEntry(const FName InName, const FNewToolMenuSectionDelegate& InConstruct)
+{
+	FToolMenuEntry& Entry = AddEntry(FToolMenuEntry(UToolMenus::Get()->CurrentOwner(), InName, EMultiBlockType::MenuEntry));
+	Entry.Construct = InConstruct;
+	return Entry;
+}
+
+FToolMenuEntry& FToolMenuSection::AddDynamicEntry(const FName InName, const FNewToolMenuDelegateLegacy& InConstruct)
+{
+	FToolMenuEntry& Entry = AddEntry(FToolMenuEntry(UToolMenus::Get()->CurrentOwner(), InName, EMultiBlockType::MenuEntry));
+	Entry.ConstructLegacy = InConstruct;
+	return Entry;
+}
+
+FToolMenuEntry& FToolMenuSection::AddMenuSeparator(const FName InName)
+{
+	return AddEntry(FToolMenuEntry::InitMenuSeparator(InName));
+}
+
+int32 FToolMenuSection::IndexOfBlock(const FName InName) const
+{
+	for (int32 i=0; i < Blocks.Num(); ++i)
+	{
+		if (Blocks[i].Name == InName)
+		{
+			return i;
+		}
+	}
+
+	return INDEX_NONE;
+}
+
+bool FToolMenuSection::IsNonLegacyDynamic() const
+{
+	return ToolMenuSectionDynamic || Construct.NewToolMenuDelegate.IsBound();
+}
+
+void FToolMenuSection::AssembleBlock(const FToolMenuEntry& Block)
+{
+	const EToolMenuInsertType Position = Block.InsertPosition.Position;
+
+	int32 ExistingIndex = IndexOfBlock(Block.Name);
+	if (ExistingIndex != INDEX_NONE)
+	{
+		Blocks[ExistingIndex] = Block;
+	}
+	else if (Position == EToolMenuInsertType::Before || Position == EToolMenuInsertType::After)
+	{
+		int32 DestIndex = IndexOfBlock(Block.InsertPosition.Name);
+		if (DestIndex != INDEX_NONE)
+		{
+			if (Position == EToolMenuInsertType::After)
+			{
+				++DestIndex;
+			}
+
+			Blocks.Insert(Block, DestIndex);
+		}
+	}
+	else if (Position == EToolMenuInsertType::First)
+	{
+		Blocks.Insert(Block, 0);
+	}
+	else
+	{
+		Blocks.Add(Block);
+	}
+}
+
+int32 FToolMenuSection::RemoveEntry(const FName InName)
+{
+	return Blocks.RemoveAll([InName](const FToolMenuEntry& Block) { return Block.Name == InName; });
+}
+
+int32 FToolMenuSection::RemoveEntriesByOwner(const FToolMenuOwner InOwner)
+{
+	if (InOwner != FToolMenuOwner())
+	{
+		return Blocks.RemoveAll([InOwner](const FToolMenuEntry& Block) { return Block.Owner == InOwner; });
+	}
+
+	return 0;
+}
+
+int32 FToolMenuSection::FindBlockInsertIndex(const FToolMenuEntry& InBlock) const
+{
+	const FToolMenuInsert InPosition = InBlock.InsertPosition;
+
+	if (InPosition.IsDefault())
+	{
+		return Blocks.Num();
+	}
+
+	if (InPosition.Position == EToolMenuInsertType::First)
+	{
+		for (int32 i = 0; i < Blocks.Num(); ++i)
+		{
+			if (Blocks[i].InsertPosition != InPosition)
+			{
+				return i;
+			}
+		}
+
+		return Blocks.Num();
+	}
+
+	int32 DestIndex = IndexOfBlock(InPosition.Name);
+	if (DestIndex == INDEX_NONE)
+	{
+		return INDEX_NONE;
+	}
+
+	if (InsertPosition.Position == EToolMenuInsertType::After)
+	{
+		++DestIndex;
+	}
+
+	for (int32 i = DestIndex; i < Blocks.Num(); ++i)
+	{
+		if (Blocks[i].InsertPosition != InPosition)
+		{
+			return i;
+		}
+	}
+
+	return Blocks.Num();
+}
diff --git a/Engine/Source/Developer/ToolMenus/Private/ToolMenus.cpp b/Engine/Source/Developer/ToolMenus/Private/ToolMenus.cpp
new file mode 100644
--- /dev/null
+++ b/Engine/Source/Developer/ToolMenus/Private/ToolMenus.cpp
@@ -0,0 +1,1496 @@
+// Copyright 1998-2019 Epic Games, Inc. All Rights Reserved.
+
+#include "ToolMenus.h"
+#include "IToolMenusModule.h"
+
+#include "Textures/SlateIcon.h"
+#include "Framework/Commands/UIAction.h"
+#include "Framework/MultiBox/MultiBox.h"
+#include "Framework/MultiBox/MultiBoxBuilder.h"
+#include "Framework/Application/SlateApplication.h"
+#include "Internationalization/Internationalization.h"
+
+#include "HAL/PlatformApplicationMisc.h" // For clipboard
+
+#define LOCTEXT_NAMESPACE "ToolMenuSubsystem"
+
+FAutoConsoleCommand ToolMenusRefreshMenuWidget = FAutoConsoleCommand(
+	TEXT("ToolMenus.RefreshAllWidgets"),
+	TEXT("Refresh All Tool Menu Widgets"),
+	FConsoleCommandDelegate::CreateLambda([]() {
+		UToolMenus::Get()->RefreshAllWidgets();
+	}));
+
+FName FToolMenuStringCommand::GetTypeName() const
+{
+	static const FName CommandName("Command");
+	static const FName PythonName("Python");
+
+	switch (Type)
+	{
+	case EToolMenuStringCommandType::Command:
+		return CommandName;
+	case EToolMenuStringCommandType::Python:
+		return PythonName;
+	case EToolMenuStringCommandType::Custom:
+		return CustomType;
+	default:
+		break;
+	}
+
+	return NAME_None;
+}
+
+FExecuteAction FToolMenuStringCommand::ToExecuteAction(const FToolMenuContext& Context) const
+{
+	if (IsBound())
+	{
+		return FExecuteAction::CreateStatic(&UToolMenus::ExecuteStringCommand, *this, Context);
+	}
+
+	return FExecuteAction();
+}
+
+FToolUIActionChoice::FToolUIActionChoice(const TSharedPtr< const FUICommandInfo >& InCommand, const FUICommandList& InCommandList)
+{
+	if (InCommand.IsValid())
+	{
+		if (const FUIAction* UIAction = InCommandList.GetActionForCommand(InCommand))
+		{
+			Action = *UIAction;
+			ToolAction.Reset();
+			DynamicToolAction.Reset();
+		}
+	}
+}
+
+UToolMenus::UToolMenus() :
+	bNextTickTimerIsSet(false),
+	bRefreshWidgetsNextTick(false),
+	bCleanupStaleWidgetsNextTick(false)
+{
+}
+
+UToolMenus* UToolMenus::Get()
+{
+	static UToolMenus* Singleton = nullptr;
+	if (!Singleton)
+	{
+		Singleton = NewObject<UToolMenus>();
+		Singleton->AddToRoot();
+		check(Singleton);
+	}
+	return Singleton;
+}
+
+bool UToolMenus::IsToolMenuUIEnabled()
+{
+	if (!FSlateApplication::IsInitialized())
+	{
+		return false;
+	}
+
+	return !IsRunningCommandlet() && !IsRunningGame() && !IsRunningDedicatedServer() && !IsRunningClientOnly();
+}
+
+FName UToolMenus::JoinMenuPaths(const FName Base, const FName Child)
+{
+	return *(Base.ToString() + TEXT(".") + Child.ToString());
+}
+
+bool UToolMenus::GetDisplayUIExtensionPoints() const
+{
+	return ShouldDisplayExtensionPoints.IsBound() && ShouldDisplayExtensionPoints.Execute();
+}
+
+UToolMenu* UToolMenus::FindMenu(const FName Name)
+{
+	UToolMenu** Found = Menus.Find(Name);
+	return Found ? *Found : nullptr;
+}
+
+bool UToolMenus::IsMenuRegistered(const FName Name) const
+{
+	const UToolMenu* const * Found = Menus.Find(Name);
+	return Found && *Found && (*Found)->IsRegistered();
+}
+
+TArray<UToolMenu*> UToolMenus::CollectHierarchy(const FName InName)
+{
+	TArray<UToolMenu*> Result;
+
+	UToolMenu* Current = FindMenu(InName);
+	while (Current)
+	{
+		// Detect infinite loop
+		for (UToolMenu* Other : Result)
+		{
+			if (Other->MenuName == Current->MenuName)
+			{
+				UE_LOG(LogToolMenus, Warning, TEXT("Infinite loop detected in tool menu: %s"), *InName.ToString());
+				return TArray<UToolMenu*>();
+			}
+		}
+
+		Result.Add(Current);
+
+		if (Current->MenuParent != NAME_None)
+		{
+			Current = FindMenu(Current->MenuParent);
+		}
+		else
+		{
+			break;
+		}
+	}
+
+	Algo::Reverse(Result);
+
+	return Result;
+}
+
+void UToolMenus::ListAllParents(const FName InName, TArray<FName>& AllParents)
+{
+	for (const UToolMenu* Menu : CollectHierarchy(InName))
+	{
+		AllParents.Add(Menu->MenuName);
+	}
+}
+
+void UToolMenus::AssembleMenuByName(UToolMenu* GeneratedMenu, const FName InName)
+{
+	if (UToolMenu* Menu = FindMenu(InName))
+	{
+		GeneratedMenu->InitGeneratedCopy(Menu);
+		AssembleMenuHierarchy(GeneratedMenu, CollectHierarchy(InName));
+	}
+}
+
+void UToolMenus::AssembleMenuSection(UToolMenu* GeneratedMenu, const UToolMenu* Other, FToolMenuSection* DestSection, const FToolMenuSection& OtherSection)
+{
+	// Build list of blocks in expected order including blocks created by construct delegates
+	TArray<FToolMenuEntry> RemainingBlocks;
+
+	UToolMenu* ConstructedEntries = nullptr;
+	for (const FToolMenuEntry& Block : OtherSection.Blocks)
+	{
+		if (!Block.IsNonLegacyDynamicConstruct())
+		{
+			RemainingBlocks.Add(Block);
+			continue;
+		}
+
+		if (ConstructedEntries == nullptr)
+		{
+			ConstructedEntries = NewObject<UToolMenu>(this);
+			ConstructedEntries->Context = DestSection->Context;
+		}
+
+		TArray<FToolMenuEntry> GeneratedEntries;
+		GeneratedEntries.Add(Block);
+
+		int32 NumIterations = 0;
+		while (GeneratedEntries.Num() > 0)
+		{
+			FToolMenuEntry& GeneratedEntry = GeneratedEntries[0];
+			if (GeneratedEntry.IsNonLegacyDynamicConstruct())
+			{
+				if (NumIterations++ > 5000)
+				{
+					UE_LOG(LogToolMenus, Warning, TEXT("Possible infinite loop for menu: %s, section: %s, block: %s"), *Other->MenuName.ToString(), *OtherSection.Name.ToString(), *Block.Name.ToString());
+					break;
+				}
+				
+				ConstructedEntries->Sections.Reset();
+				if (GeneratedEntry.IsScriptObjectDynamicConstruct())
+				{
+					GeneratedEntry.ScriptObject->ConstructMenuEntry(ConstructedEntries, DestSection->Name, DestSection->Context);
+				}
+				else
+				{
+					FToolMenuSection& ConstructedSection = ConstructedEntries->AddSection(DestSection->Name);
+					ConstructedSection.Context = ConstructedEntries->Context;
+					GeneratedEntry.Construct.Execute(ConstructedSection);
+				}
+				GeneratedEntries.RemoveAt(0, 1, false);
+
+				// Combine all user's choice of selections here into the current section target
+				// If the user wants to add items to different sections they will need to create dynamic section instead (for now)
+				int32 NumBlocksInserted = 0;
+				for (FToolMenuSection& ConstructedSection : ConstructedEntries->Sections)
+				{
+					for (FToolMenuEntry& ConstructedBlock : ConstructedSection.Blocks)
+					{
+						if (ConstructedBlock.InsertPosition.IsDefault())
+						{
+							ConstructedBlock.InsertPosition = Block.InsertPosition;
+						}
+					}
+					GeneratedEntries.Insert(ConstructedSection.Blocks, NumBlocksInserted);
+					NumBlocksInserted += ConstructedSection.Blocks.Num();
+				}
+			}
+			else
+			{
+				RemainingBlocks.Add(GeneratedEntry);
+				GeneratedEntries.RemoveAt(0, 1, false);
+			}
+		}
+	}
+
+	if (ConstructedEntries)
+	{
+		ConstructedEntries->Sections.Empty();
+		ConstructedEntries = nullptr;
+	}
+
+	// Repeatedly loop because insert location may not exist until later in list
+	while (RemainingBlocks.Num() > 0)
+	{
+		int32 NumHandled = 0;
+		for (int32 i = 0; i < RemainingBlocks.Num(); ++i)
+		{
+			FToolMenuEntry& Block = RemainingBlocks[i];
+			int32 DestIndex = DestSection->FindBlockInsertIndex(Block);
+			if (DestIndex != INDEX_NONE)
+			{
+				DestSection->Blocks.Insert(Block, DestIndex);
+				RemainingBlocks.RemoveAt(i);
+				--i;
+				++NumHandled;
+				// Restart loop because items earlier in the list may need to attach to this block
+				break;
+			}
+		}
+		if (NumHandled == 0)
+		{
+			for (const FToolMenuEntry& Block : RemainingBlocks)
+			{
+				UE_LOG(LogToolMenus, Warning, TEXT("Menu item not found: '%s' for insert: '%s'"), *Block.InsertPosition.Name.ToString(), *Block.Name.ToString());
+			}
+			break;
+		}
+	}
+}
+
+void UToolMenus::AssembleMenu(UToolMenu* GeneratedMenu, const UToolMenu* Other)
+{
+	TArray<FToolMenuSection> RemainingSections;
+
+	UToolMenu* ConstructedSections = nullptr;
+	for (const FToolMenuSection& OtherSection : Other->Sections)
+	{
+		if (!OtherSection.IsNonLegacyDynamic())
+		{
+			RemainingSections.Add(OtherSection);
+			continue;
+		}
+		
+		if (ConstructedSections == nullptr)
+		{
+			ConstructedSections = NewObject<UToolMenu>(this);
+			ConstructedSections->Context = GeneratedMenu->Context;
+		}
+
+		TArray<FToolMenuSection> GeneratedSections;
+		GeneratedSections.Add(OtherSection);
+
+		int32 NumIterations = 0;
+		while (GeneratedSections.Num() > 0)
+		{
+			if (GeneratedSections[0].IsNonLegacyDynamic())
+			{
+				if (NumIterations++ > 5000)
+				{
+					UE_LOG(LogToolMenus, Warning, TEXT("Possible infinite loop for menu: %s, section: %s"), *Other->MenuName.ToString(), *OtherSection.Name.ToString());
+					break;
+				}
+
+				ConstructedSections->Sections.Reset();
+				
+				if (GeneratedSections[0].ToolMenuSectionDynamic)
+				{
+					GeneratedSections[0].ToolMenuSectionDynamic->ConstructSections(ConstructedSections, GeneratedMenu->Context);
+				}
+				else if (GeneratedSections[0].Construct.NewToolMenuDelegate.IsBound())
+				{
+					GeneratedSections[0].Construct.NewToolMenuDelegate.Execute(ConstructedSections);
+				}
+
+				for (FToolMenuSection& ConstructedSection : ConstructedSections->Sections)
+				{
+					if (ConstructedSection.InsertPosition.IsDefault())
+					{
+						ConstructedSection.InsertPosition = GeneratedSections[0].InsertPosition;
+					}
+				}
+				
+				GeneratedSections.RemoveAt(0, 1, false);				
+				GeneratedSections.Insert(ConstructedSections->Sections, 0);
+			}
+			else
+			{
+				RemainingSections.Add(GeneratedSections[0]);
+				GeneratedSections.RemoveAt(0, 1, false);
+			}
+		}
+	}
+
+	if (ConstructedSections)
+	{
+		ConstructedSections->Sections.Empty();
+		ConstructedSections = nullptr;
+	}
+
+	while (RemainingSections.Num() > 0)
+	{
+		int32 NumHandled = 0;
+		for (int32 i=0; i < RemainingSections.Num(); ++i)
+		{
+			const FToolMenuSection& RemainingSection = RemainingSections[i];
+
+			// Update existing section
+			FToolMenuSection* Section = GeneratedMenu->FindSection(RemainingSection.Name);
+			if (!Section)
+			{
+				// Try add new section (if insert location exists)
+				int32 DestIndex = GeneratedMenu->FindInsertIndex(RemainingSection);
+				if (DestIndex != INDEX_NONE)
+				{
+					GeneratedMenu->Sections.InsertDefaulted(DestIndex);
+					Section = &GeneratedMenu->Sections[DestIndex];
+					Section->InitGeneratedSectionCopy(RemainingSection, GeneratedMenu->Context);
+				}
+				else
+				{
+					continue;
+				}
+			}
+
+			AssembleMenuSection(GeneratedMenu, Other, Section, RemainingSection);
+			RemainingSections.RemoveAt(i);
+			--i;
+			++NumHandled;
+			break;
+		}
+		if (NumHandled == 0)
+		{
+			for (const FToolMenuSection& RemainingSection : RemainingSections)
+			{
+				UE_LOG(LogToolMenus, Warning, TEXT("Menu section not found: '%s' for insert: '%s'"), *RemainingSection.InsertPosition.Name.ToString(), *RemainingSection.Name.ToString());
+			}
+			break;
+		}
+	}
+}
+
+int32 UToolMenus::FindCustomizedMenuIndex(const FName InName)
+{
+	for (int32 i = 0; i < CustomizedMenus.Num(); ++i)
+	{
+		if (CustomizedMenus[i].Name == InName)
+		{
+			return i;
+		}
+	}
+
+	return INDEX_NONE;
+}
+
+FCustomizedToolMenu* UToolMenus::FindCustomizedMenu(const FName InName)
+{
+	int32 FoundIndex = FindCustomizedMenuIndex(InName);
+	return (FoundIndex != INDEX_NONE) ? &CustomizedMenus[FoundIndex] : nullptr;
+}
+
+void UToolMenus::ApplyCustomization(UToolMenu* GeneratedMenu)
+{
+	const FCustomizedToolMenu* CustomizedMenu = FindCustomizedMenu(GeneratedMenu->MenuName);
+	if (CustomizedMenu == nullptr)
+	{
+		return;
+	}
+
+	TArray<FToolMenuSection> NewSections;
+	NewSections.Reserve(GeneratedMenu->Sections.Num());
+
+	for (const FCustomizedToolMenuSection& CustomizedSection : CustomizedMenu->Sections)
+	{
+		int32 SectionIndex = GeneratedMenu->IndexOfSection(CustomizedSection.Name);
+		if (SectionIndex == INDEX_NONE)
+		{
+			continue;
+		}
+
+		FToolMenuSection& Section = GeneratedMenu->Sections[SectionIndex];
+
+		TArray<FToolMenuEntry> NewBlocks;
+		NewBlocks.Reserve(Section.Blocks.Num());
+
+		for (const FName ItemName : CustomizedSection.Items)
+		{
+			int32 EntrySectionIndex = INDEX_NONE;
+			int32 EntryIndex = INDEX_NONE;
+			if (GeneratedMenu->FindEntry(ItemName, EntrySectionIndex, EntryIndex))
+			{
+				NewBlocks.Add(GeneratedMenu->Sections[EntrySectionIndex].Blocks[EntryIndex]);
+				GeneratedMenu->Sections[EntrySectionIndex].Blocks.RemoveAt(EntryIndex);
+			}
+		}
+
+		// Remaining blocks
+		for (FToolMenuEntry& Block : Section.Blocks)
+		{
+			NewBlocks.Add(Block);
+		}
+
+		Section.Blocks = NewBlocks;
+
+		NewSections.Add(Section);
+		GeneratedMenu->Sections.RemoveAt(SectionIndex);
+	}
+
+	// Remaining sections
+	NewSections.Append(GeneratedMenu->Sections);
+
+	// Hide
+	for (int32 SectionIndex=0; SectionIndex < NewSections.Num(); ++SectionIndex)
+	{
+		FToolMenuSection& Section = NewSections[SectionIndex];
+		if (CustomizedMenu->HiddenSections.Contains(Section.Name))
+		{
+			NewSections.RemoveAt(SectionIndex);
+			--SectionIndex;
+			continue;
+		}
+
+		for (int32 i = 0; i < Section.Blocks.Num(); ++i)
+		{
+			if (CustomizedMenu->HiddenItems.Contains(Section.Blocks[i].Name))
+			{
+				Section.Blocks.RemoveAt(i);
+				--i;
+			}
+		}
+	}
+
+	GeneratedMenu->Sections = NewSections;
+}
+
+void UToolMenus::AssembleMenuHierarchy(UToolMenu* GeneratedMenu, const TArray<UToolMenu*>& Hierarchy)
+{
+	if (GeneratedMenu->MenuType == EMultiBoxType::MenuBar)
+	{
+		// Menu Bars require one section
+		if (GeneratedMenu->Sections.Num() == 0)
+		{
+			GeneratedMenu->Sections.AddDefaulted();
+		}
+
+		FToolMenuSection& MenuBarSection = GeneratedMenu->Sections[0];
+		for (const UToolMenu* MenuData : Hierarchy)
+		{
+			for (const FToolMenuSection& Section : MenuData->Sections)
+			{
+				for (const FToolMenuEntry& Block : Section.Blocks)
+				{
+					MenuBarSection.AssembleBlock(Block);
+				}
+			}
+		}
+	}
+	else
+	{
+		for (const UToolMenu* FoundParent : Hierarchy)
+		{
+			AssembleMenu(GeneratedMenu, FoundParent);
+		}
+	}
+
+	ApplyCustomization(GeneratedMenu);
+}
+
+void UToolMenus::FillMenuDynamic(FMenuBuilder& MenuBuilder, FNewToolMenuDelegate InConstructMenu, const FToolMenuContext Context)
+{
+	if (InConstructMenu.IsBound())
+	{
+		// Create final menu
+		UToolMenu* MenuData = NewObject<UToolMenu>();
+		MenuData->Context = Context;
+		InConstructMenu.Execute(MenuData);
+
+		// Populate menu builder with final menu
+		PopulateMenuBuilder(MenuBuilder, MenuData);
+	}
+}
+
+void UToolMenus::FillMenu(class FMenuBuilder& MenuBuilder, FName InMenuName, FToolMenuContext InMenuContext)
+{
+	// Create combined final menu
+	UToolMenu* GeneratedMenu = NewObject<UToolMenu>();
+	GeneratedMenu->Context = InMenuContext;
+	AssembleMenuByName(GeneratedMenu, InMenuName);
+
+	// Populate menu builder with final menu
+	PopulateMenuBuilder(MenuBuilder, GeneratedMenu);
+}
+
+TSharedRef<SWidget> UToolMenus::GenerateToolbarComboButtonMenu(const FName SubMenuFullName, FToolMenuContext InContext)
+{
+	return GenerateWidget(SubMenuFullName, InContext);
+}
+
+void UToolMenus::FillMenuBarDropDown(class FMenuBuilder& MenuBuilder, FName InParentName, FName InChildName, FToolMenuContext InMenuContext)
+{
+	if (UToolMenu* MenuToUse = FindSubMenuToGenerateWith(InParentName, InChildName))
+	{
+		// Create combined final menu
+		UToolMenu* GeneratedMenu = NewObject<UToolMenu>();
+		GeneratedMenu->Context = InMenuContext;
+		AssembleMenuByName(GeneratedMenu, MenuToUse->MenuName);
+		GeneratedMenu->MenuName = JoinMenuPaths(InParentName, InChildName);
+
+		// Populate menu builder with final menu
+		PopulateMenuBuilder(MenuBuilder, GeneratedMenu);
+	}
+}
+
+void UToolMenus::PopulateMenuBuilder(FMenuBuilder& MenuBuilder, UToolMenu* MenuData)
+{
+	if (GetDisplayUIExtensionPoints())
+	{
+		const FName MenuName = MenuData->GetMenuName();
+		MenuBuilder.AddMenuEntry(
+			FText::FromName(MenuName),
+			LOCTEXT("CopyMenuNameToClipboard", "Copy menu name to clipboard"),
+			FSlateIcon(),
+			FExecuteAction::CreateLambda([MenuName]() { FPlatformApplicationMisc::ClipboardCopy(*MenuName.ToString()); }),
+			"MenuName"
+		);
+	}
+
+	for (FToolMenuSection& Section : MenuData->Sections)
+	{
+		if (Section.Construct.NewToolMenuDelegateLegacy.IsBound())
+		{
+			Section.Construct.NewToolMenuDelegateLegacy.Execute(MenuBuilder, MenuData);
+			continue;
+		}
+
+		MenuBuilder.BeginSection(Section.Name, Section.Label);
+
+		for (FToolMenuEntry& Block : Section.Blocks)
+		{
+			if (Block.ConstructLegacy.IsBound())
+			{
+				Block.ConstructLegacy.Execute(MenuBuilder, MenuData);
+				continue;
+			}
+
+			FUIAction UIAction = ConvertUIAction(Block, MenuData->Context);
+
+			TSharedPtr<SWidget> Widget;
+			if (Block.MakeWidget.IsBound())
+			{
+				Widget = Block.MakeWidget.Execute(MenuData->Context);
+			}
+
+			if (Block.Type == EMultiBlockType::MenuEntry)
+			{
+				if (Block.IsSubMenu())
+				{
+					if (Block.SubMenuData.ConstructMenu.NewMenuDelegate.IsBound())
+					{
+						MenuBuilder.AddSubMenu(
+							Block.Label,
+							Block.ToolTip,
+							Block.SubMenuData.ConstructMenu.NewMenuDelegate,
+							Block.SubMenuData.bOpenSubMenuOnClick,
+							Block.Icon.Get(),
+							Block.bShouldCloseWindowAfterMenuSelection,
+							Block.Name
+						);
+					}
+					else if (Block.SubMenuData.ConstructMenu.NewToolMenuDelegate.IsBound())
+					{
+						// SubMenu constructed each time it is opened
+						MenuBuilder.AddSubMenu(
+							Block.Label,
+							Block.ToolTip,
+							FNewMenuDelegate::CreateUObject(this, &UToolMenus::FillMenuDynamic, Block.SubMenuData.ConstructMenu.NewToolMenuDelegate, MenuData->Context),
+							Block.SubMenuData.bOpenSubMenuOnClick,
+							Block.Icon.Get(),
+							Block.bShouldCloseWindowAfterMenuSelection,
+							Block.Name
+						);
+					}
+					else
+					{
+						// SubMenu registered once by name in database
+						FName SubMenuFullName = JoinMenuPaths(MenuData->MenuName, Block.Name);
+						FNewMenuDelegate NewMenuDelegate = FNewMenuDelegate::CreateUObject(this, &UToolMenus::FillMenu, SubMenuFullName, MenuData->Context);
+
+						if (Widget.IsValid())
+						{
+							// Could also check if Visible/Enabled bound as well
+							if (UIAction.IsBound())
+							{
+								MenuBuilder.AddSubMenu(UIAction, Widget.ToSharedRef(), NewMenuDelegate, Block.bShouldCloseWindowAfterMenuSelection);
+							}
+							else
+							{
+								MenuBuilder.AddSubMenu(Widget.ToSharedRef(), NewMenuDelegate, Block.SubMenuData.bOpenSubMenuOnClick, Block.bShouldCloseWindowAfterMenuSelection);
+							}
+						}
+						else
+						{
+							MenuBuilder.AddSubMenu(
+								Block.Label,
+								Block.ToolTip,
+								NewMenuDelegate,
+								Block.SubMenuData.bOpenSubMenuOnClick,
+								Block.Icon.Get(),
+								Block.bShouldCloseWindowAfterMenuSelection,
+								Block.Name
+							);
+						}
+					}
+				}
+				else
+				{
+					if (Block.Command.IsValid())
+					{
+						bool bPopCommandList = false;
+						TSharedPtr<const FUICommandList> CommandListForAction;
+						if (Block.GetActionForCommand(MenuData->Context, CommandListForAction) != nullptr && CommandListForAction.IsValid())
+						{
+							MenuBuilder.PushCommandList(CommandListForAction.ToSharedRef());
+							bPopCommandList = true;
+						}
+						else
+						{
+							UE_LOG(LogToolMenus, Error, TEXT("UI command not found for menu entry: %s, menu: %s"), *Block.Name.ToString(), *MenuData->MenuName.ToString());
+						}
+
+						MenuBuilder.AddMenuEntry(Block.Command, Block.Name, Block.Label, Block.ToolTip, Block.Icon.Get());
+
+						if (bPopCommandList)
+						{
+							MenuBuilder.PopCommandList();
+						}
+					}
+					else if (Block.ScriptObject)
+					{
+						UToolMenuEntryScript* ScriptObject = Block.ScriptObject;
+						const FSlateIcon Icon = ScriptObject->CreateIconAttribute(MenuData->Context).Get();
+						MenuBuilder.AddMenuEntry(ScriptObject->CreateLabelAttribute(MenuData->Context), ScriptObject->CreateToolTipAttribute(MenuData->Context), Icon, UIAction, ScriptObject->Data.Name, Block.UserInterfaceActionType, Block.TutorialHighlightName);
+					}
+					else
+					{
+						if (Widget.IsValid())
+						{
+							MenuBuilder.AddMenuEntry(UIAction, Widget.ToSharedRef(), Block.Name, Block.ToolTip, Block.UserInterfaceActionType, Block.TutorialHighlightName);
+						}
+						else
+						{
+							MenuBuilder.AddMenuEntry(Block.Label, Block.ToolTip, Block.Icon.Get(), UIAction, Block.Name, Block.UserInterfaceActionType, Block.TutorialHighlightName);
+						}
+					}
+				}
+			}
+			else if (Block.Type == EMultiBlockType::MenuSeparator)
+			{
+				MenuBuilder.AddMenuSeparator(Block.Name);
+			}
+			else if (Block.Type == EMultiBlockType::Widget)
+			{
+				MenuBuilder.AddWidget(Widget.ToSharedRef(), Block.Label.Get(), Block.WidgetData.bNoIndent, Block.WidgetData.bSearchable);
+			}
+			else
+			{
+				UE_LOG(LogToolMenus, Warning, TEXT("Menu '%s', item '%s', type not currently supported: %d"), *MenuData->MenuName.ToString(), *Block.Name.ToString(), Block.Type);
+			}
+		}
+
+		MenuBuilder.EndSection();
+	}
+
+	AddReferencedContextObjects(MenuBuilder.GetMultiBox(), MenuData->Context);
+}
+
+void UToolMenus::PopulateToolBarBuilder(FToolBarBuilder& ToolBarBuilder, UToolMenu* MenuData)
+{
+	for (FToolMenuSection& Section : MenuData->Sections)
+	{
+		if (Section.Construct.NewToolBarDelegateLegacy.IsBound())
+		{
+			Section.Construct.NewToolBarDelegateLegacy.Execute(ToolBarBuilder, MenuData);
+			continue;
+		}
+
+		ToolBarBuilder.BeginSection(Section.Name);
+
+		for (FToolMenuEntry& Block : Section.Blocks)
+		{
+			if (Block.ToolBarData.ConstructLegacy.IsBound())
+			{
+				Block.ToolBarData.ConstructLegacy.Execute(ToolBarBuilder, MenuData);
+				continue;
+			}
+
+			FUIAction UIAction = ConvertUIAction(Block, MenuData->Context);
+
+			TSharedPtr<SWidget> Widget;
+			if (Block.MakeWidget.IsBound())
+			{
+				Widget = Block.MakeWidget.Execute(MenuData->Context);
+			}
+
+			if (Block.Type == EMultiBlockType::ToolBarButton)
+			{
+				if (Block.Command.IsValid())
+				{
+					bool bPopCommandList = false;
+					TSharedPtr<const FUICommandList> CommandListForAction;
+					if (Block.GetActionForCommand(MenuData->Context, CommandListForAction) != nullptr && CommandListForAction.IsValid())
+					{
+						ToolBarBuilder.PushCommandList(CommandListForAction.ToSharedRef());
+						bPopCommandList = true;
+					}
+					else
+					{
+						UE_LOG(LogToolMenus, Error, TEXT("UI command not found for toolbar entry: %s, toolbar: %s"), *Block.Name.ToString(), *MenuData->MenuName.ToString());
+					}
+
+					ToolBarBuilder.AddToolBarButton(Block.Command, Block.Name, Block.Label, Block.ToolTip, Block.Icon, Block.TutorialHighlightName);
+
+					if (bPopCommandList)
+					{
+						ToolBarBuilder.PopCommandList();
+					}
+				}
+				else if (Block.ScriptObject)
+				{
+					UToolMenuEntryScript* ScriptObject = Block.ScriptObject;
+					TAttribute<FSlateIcon> Icon = ScriptObject->CreateIconAttribute(MenuData->Context);
+					ToolBarBuilder.AddToolBarButton(UIAction, ScriptObject->Data.Name, ScriptObject->CreateLabelAttribute(MenuData->Context), ScriptObject->CreateToolTipAttribute(MenuData->Context), Icon, Block.UserInterfaceActionType, Block.TutorialHighlightName);
+				}
+				else
+				{
+					ToolBarBuilder.AddToolBarButton(UIAction, Block.Name, Block.Label, Block.ToolTip, Block.Icon, Block.UserInterfaceActionType, Block.TutorialHighlightName);
+				}
+			}
+			else if (Block.Type == EMultiBlockType::ToolBarComboButton)
+			{
+				FOnGetContent OnGetContent = ConvertWidgetChoice(Block.ToolBarData.ComboButtonContextMenuGenerator, MenuData->Context);
+				if (OnGetContent.IsBound())
+				{
+					ToolBarBuilder.AddComboButton(UIAction, OnGetContent, Block.Label, Block.ToolTip, Block.Icon, Block.ToolBarData.bSimpleComboBox, Block.TutorialHighlightName);
+				}
+				else
+				{
+					FName SubMenuFullName = JoinMenuPaths(MenuData->MenuName, Block.Name);
+					FOnGetContent Delegate = FOnGetContent::CreateUObject(this, &UToolMenus::GenerateToolbarComboButtonMenu, SubMenuFullName, MenuData->Context);
+					ToolBarBuilder.AddComboButton(UIAction, Delegate, Block.Label, Block.ToolTip, Block.Icon, Block.ToolBarData.bSimpleComboBox, Block.TutorialHighlightName);
+				}
+			}
+			else if (Block.Type == EMultiBlockType::ToolBarSeparator)
+			{
+				ToolBarBuilder.AddSeparator(Block.Name);
+			}
+			else if (Block.Type == EMultiBlockType::Widget)
+			{
+				ToolBarBuilder.AddWidget(Widget.ToSharedRef(), Block.TutorialHighlightName, Block.WidgetData.bSearchable);
+			}
+			else
+			{
+				UE_LOG(LogToolMenus, Warning, TEXT("Toolbar '%s', item '%s', type not currently supported: %d"), *MenuData->MenuName.ToString(), *Block.Name.ToString(), Block.Type);
+			}
+		}
+
+		ToolBarBuilder.EndSection();
+	}
+
+	if (GetDisplayUIExtensionPoints())
+	{
+		const FName MenuName = MenuData->GetMenuName();
+		ToolBarBuilder.BeginSection(MenuName);
+		ToolBarBuilder.AddToolBarButton(
+			FExecuteAction::CreateLambda([MenuName]() { FPlatformApplicationMisc::ClipboardCopy(*MenuName.ToString()); }), 
+			"MenuName",
+			LOCTEXT("CopyNameToClipboard", "Copy Name"),
+			LOCTEXT("CopyMenuNameToClipboard", "Copy menu name to clipboard")
+		);
+		ToolBarBuilder.EndSection();
+	}
+
+	AddReferencedContextObjects(ToolBarBuilder.GetMultiBox(), MenuData->Context);
+}
+
+void UToolMenus::PopulateMenuBarBuilder(FMenuBarBuilder& MenuBarBuilder, UToolMenu* GeneratedMenu)
+{
+	if (GeneratedMenu->Sections.Num() > 0)
+	{
+		for (FToolMenuEntry& Block : GeneratedMenu->Sections[0].Blocks)
+		{
+			MenuBarBuilder.AddPullDownMenu(
+				Block.Label.Get(),
+				Block.ToolTip.Get(),
+				FNewMenuDelegate::CreateUObject(this, &UToolMenus::FillMenuBarDropDown, GeneratedMenu->MenuName, Block.Name, GeneratedMenu->Context),
+				Block.Name
+			);
+		}
+
+		AddReferencedContextObjects(MenuBarBuilder.GetMultiBox(), GeneratedMenu->Context);
+	}
+}
+
+FOnGetContent UToolMenus::ConvertWidgetChoice(const FNewToolMenuWidgetChoice& Choice, const FToolMenuContext& Context) const
+{
+	if (Choice.NewToolMenuWidget.IsBound())
+	{
+		return FOnGetContent::CreateLambda([ToCall = Choice.NewToolMenuWidget, Context]()
+		{
+			if (ToCall.IsBound())
+			{
+				return ToCall.Execute(Context);
+			}
+
+			return SNullWidget::NullWidget;
+		});
+	}
+	else if (Choice.NewToolMenu.IsBound())
+	{
+		return FOnGetContent::CreateLambda([ToCall = Choice.NewToolMenu, Context]()
+		{
+			if (ToCall.IsBound())
+			{
+				UToolMenu* MenuData = NewObject<UToolMenu>();
+				MenuData->Context = Context;
+				ToCall.Execute(MenuData);
+				return UToolMenus::Get()->GenerateWidget(MenuData);
+			}
+
+			return SNullWidget::NullWidget;
+		});
+	}
+	return Choice.OnGetContent;
+}
+
+FUIAction UToolMenus::ConvertUIAction(const FToolMenuEntry& Block, const FToolMenuContext& Context)
+{
+	FUIAction UIAction;
+	
+	if (Block.ScriptObject)
+	{
+		UIAction = ConvertScriptObjectToUIAction(Block.ScriptObject, Context);
+	}
+	else
+	{
+		UIAction = ConvertUIAction(Block.Action, Context);
+	}
+	
+	if (!UIAction.ExecuteAction.IsBound() && Block.StringExecuteAction.IsBound())
+	{
+		UIAction.ExecuteAction = Block.StringExecuteAction.ToExecuteAction(Context);
+	}
+
+	return UIAction;
+}
+
+FUIAction UToolMenus::ConvertUIAction(const FToolUIActionChoice& Choice, const FToolMenuContext& Context)
+{
+	if (const FToolUIAction* ToolAction = Choice.GetToolUIAction())
+	{
+		return ConvertUIAction(*ToolAction, Context);
+	}
+	else if (const FToolDynamicUIAction* DynamicToolAction = Choice.GetToolDynamicUIAction())
+	{
+		return ConvertUIAction(*DynamicToolAction, Context);
+	}
+	else if (const FUIAction* Action = Choice.GetUIAction())
+	{
+		return *Action;
+	}
+
+	return FUIAction();
+}
+
+FUIAction UToolMenus::ConvertUIAction(const FToolUIAction& Actions, const FToolMenuContext& Context)
+{
+	FUIAction UIAction;
+
+	if (Actions.ExecuteAction.IsBound())
+	{
+		UIAction.ExecuteAction.BindLambda([DelegateToCall = Actions.ExecuteAction, Context]()
+		{
+			DelegateToCall.ExecuteIfBound(Context);
+		});
+	}
+
+	if (Actions.CanExecuteAction.IsBound())
+	{
+		UIAction.CanExecuteAction.BindLambda([DelegateToCall = Actions.CanExecuteAction, Context]()
+		{
+			return DelegateToCall.Execute(Context);
+		});
+	}
+
+	if (Actions.GetActionCheckState.IsBound())
+	{
+		UIAction.GetActionCheckState.BindLambda([DelegateToCall = Actions.GetActionCheckState, Context]()
+		{
+			return DelegateToCall.Execute(Context);
+		});
+	}
+
+	if (Actions.IsActionVisibleDelegate.IsBound())
+	{
+		UIAction.IsActionVisibleDelegate.BindLambda([DelegateToCall = Actions.IsActionVisibleDelegate, Context]()
+		{
+			return DelegateToCall.Execute(Context);
+		});
+	}
+
+	return UIAction;
+}
+
+FUIAction UToolMenus::ConvertUIAction(const FToolDynamicUIAction& Actions, const FToolMenuContext& Context)
+{
+	FUIAction UIAction;
+
+	if (Actions.ExecuteAction.IsBound())
+	{
+		UIAction.ExecuteAction.BindLambda([DelegateToCall = Actions.ExecuteAction, Context]()
+		{
+			DelegateToCall.ExecuteIfBound(Context);
+		});
+	}
+
+	if (Actions.CanExecuteAction.IsBound())
+	{
+		UIAction.CanExecuteAction.BindLambda([DelegateToCall = Actions.CanExecuteAction, Context]()
+		{
+			return DelegateToCall.Execute(Context);
+		});
+	}
+
+	if (Actions.GetActionCheckState.IsBound())
+	{
+		UIAction.GetActionCheckState.BindLambda([DelegateToCall = Actions.GetActionCheckState, Context]()
+		{
+			return DelegateToCall.Execute(Context);
+		});
+	}
+
+	if (Actions.IsActionVisibleDelegate.IsBound())
+	{
+		UIAction.IsActionVisibleDelegate.BindLambda([DelegateToCall = Actions.IsActionVisibleDelegate, Context]()
+		{
+			return DelegateToCall.Execute(Context);
+		});
+	}
+
+	return UIAction;
+}
+
+FUIAction UToolMenus::ConvertScriptObjectToUIAction(UToolMenuEntryScript* ScriptObject, const FToolMenuContext& Context)
+{
+	FUIAction UIAction;
+
+	if (ScriptObject)
+	{
+		UClass* ScriptClass = ScriptObject->GetClass();
+
+		static const FName ExecuteName = GET_FUNCTION_NAME_CHECKED(UToolMenuEntryScript, Execute);
+		if (ScriptClass->IsFunctionImplementedInScript(ExecuteName))
+		{
+			UIAction.ExecuteAction.BindUFunction(ScriptObject, ExecuteName, Context);
+		}
+
+		static const FName CanExecuteName = GET_FUNCTION_NAME_CHECKED(UToolMenuEntryScript, CanExecute);
+		if (ScriptClass->IsFunctionImplementedInScript(CanExecuteName))
+		{
+			UIAction.CanExecuteAction.BindUFunction(ScriptObject, CanExecuteName, Context);
+		}
+
+		static const FName GetCheckStateName = GET_FUNCTION_NAME_CHECKED(UToolMenuEntryScript, GetCheckState);
+		if (ScriptClass->IsFunctionImplementedInScript(GetCheckStateName))
+		{
+			UIAction.GetActionCheckState.BindUFunction(ScriptObject, GetCheckStateName, Context);
+		}
+
+		static const FName IsVisibleName = GET_FUNCTION_NAME_CHECKED(UToolMenuEntryScript, IsVisible);
+		if (ScriptClass->IsFunctionImplementedInScript(IsVisibleName))
+		{
+			UIAction.IsActionVisibleDelegate.BindUFunction(ScriptObject, IsVisibleName, Context);
+		}
+	}
+
+	return UIAction;
+}
+
+void UToolMenus::ExecuteStringCommand(const FToolMenuStringCommand StringCommand, const FToolMenuContext Context)
+{
+	if (StringCommand.IsBound())
+	{
+		const FName TypeName = StringCommand.GetTypeName();
+		UToolMenus* ToolMenus = UToolMenus::Get();
+		if (const FToolMenuExecuteString* Handler = ToolMenus->StringCommandHandlers.Find(TypeName))
+		{
+			if (Handler->IsBound())
+			{
+				Handler->Execute(StringCommand.String, Context);
+			}
+		}
+		else
+		{
+			UE_LOG(LogToolMenus, Warning, TEXT("Unknown string command handler type: '%s'"), *TypeName.ToString());
+		}
+	}
+}
+
+UToolMenu* UToolMenus::FindSubMenuToGenerateWith(const FName InParentName, const FName InChildName)
+{
+	FName BaseName = InParentName;
+	while (BaseName != NAME_None)
+	{
+		FName JoinedName = JoinMenuPaths(BaseName, InChildName);
+		if (UToolMenu* Found = FindMenu(JoinedName))
+		{
+			return Found;
+		}
+
+		UToolMenu* BaseData = FindMenu(BaseName);
+		BaseName = BaseData ? BaseData->MenuParent : NAME_None;
+	}
+
+	return nullptr;
+}
+
+UObject* UToolMenus::FindContext(const FToolMenuContext& InContext, UClass* InClass)
+{
+	return InContext.FindByClass(InClass);
+}
+
+void UToolMenus::AddReferencedContextObjects(const TSharedRef<FMultiBox>& InMultiBox, const FToolMenuContext& InMenuContext)
+{
+	if (InMenuContext.ContextObjects.Num() == 0)
+	{
+		return;
+	}
+
+	TArray<UObject*>& References = WidgetObjectReferences.FindOrAdd(InMultiBox);
+	for (const TWeakObjectPtr<UObject>& WeakObject : InMenuContext.ContextObjects)
+	{
+		if (UObject* Object = WeakObject.Get())
+		{
+			References.AddUnique(Object);
+		}
+	}
+}
+
+void UToolMenus::AddReferencedObjects(UObject* InThis, FReferenceCollector& Collector)
+{
+	UToolMenus* This = CastChecked<UToolMenus>(InThis);
+
+	for (auto It = This->WidgetObjectReferences.CreateIterator(); It; ++It)
+	{
+		if (It->Key.IsValid())
+		{
+			Collector.AddReferencedObjects(It->Value, InThis);
+		}
+		else
+		{
+			It.RemoveCurrent();
+		}
+	}
+
+	Super::AddReferencedObjects(InThis, Collector);
+}
+
+UToolMenu* UToolMenus::GenerateMenu(const FName Name, FToolMenuContext& InMenuContext)
+{
+	return GenerateMenu(CollectHierarchy(Name), InMenuContext);
+}
+
+UToolMenu* UToolMenus::GenerateMenu(const TArray<UToolMenu*>& Hierarchy, FToolMenuContext& InMenuContext)
+{
+	UToolMenu* GeneratedMenu = NewObject<UToolMenu>(this);
+
+	if (Hierarchy.Num() > 0)
+	{
+		GeneratedMenu->InitGeneratedCopy(Hierarchy[0]);
+		GeneratedMenu->MenuName = Hierarchy.Last()->MenuName;
+		GeneratedMenu->Context = InMenuContext;
+		AssembleMenuHierarchy(GeneratedMenu, Hierarchy);
+	}
+
+	return GeneratedMenu;
+}
+
+TSharedRef< class SWidget > UToolMenus::GenerateWidget(const FName InName, FToolMenuContext& InMenuContext)
+{
+	UToolMenu* Generated = GenerateMenu(InName, InMenuContext);
+	return GenerateWidget(Generated);
+}
+
+TSharedRef<SWidget> UToolMenus::GenerateWidget(const TArray<UToolMenu*>& Hierarchy, FToolMenuContext& InMenuContext)
+{
+	if (Hierarchy.Num() == 0)
+	{
+		return SNullWidget::NullWidget;
+	}
+
+	UToolMenu* Generated = GenerateMenu(Hierarchy, InMenuContext);
+	return GenerateWidget(Generated);
+}
+
+TSharedRef<SWidget> UToolMenus::GenerateWidget(UToolMenu* GeneratedMenu)
+{
+	CleanupStaleWidgetsNextTick();
+
+	FGeneratedToolMenuWidgets& WidgetsForMenuName = GeneratedMenuWidgets.FindOrAdd(GeneratedMenu->MenuName);
+
+	// Store a copy so that we can call 'Refresh' on menus not in the database
+	FGeneratedToolMenuWidget& GeneratedMenuWidget = WidgetsForMenuName.Instances.AddDefaulted_GetRef();
+	GeneratedMenuWidget.GeneratedMenu = DuplicateObject<UToolMenu>(GeneratedMenu, this);
+	// Copy native properties that serialize does not
+	GeneratedMenuWidget.GeneratedMenu->Context = GeneratedMenu->Context;
+	GeneratedMenuWidget.GeneratedMenu->StyleSet = GeneratedMenu->StyleSet;
+
+	if (GeneratedMenu->MenuType == EMultiBoxType::Menu)
+	{
+		FMenuBuilder MenuBuilder(GeneratedMenu->bShouldCloseWindowAfterMenuSelection, GeneratedMenu->Context.CommandList, GeneratedMenu->Context.GetAllExtenders(), GeneratedMenu->bCloseSelfOnly, GeneratedMenu->StyleSet, GeneratedMenu->bSearchable);
+		PopulateMenuBuilder(MenuBuilder, GeneratedMenu);
+		TSharedRef<SWidget> Result = MenuBuilder.MakeWidget();
+		GeneratedMenuWidget.Widget = Result;
+		return Result;
+	}
+	else if (GeneratedMenu->MenuType == EMultiBoxType::MenuBar)
+	{
+		FMenuBarBuilder MenuBarBuilder(GeneratedMenu->Context.CommandList, GeneratedMenu->Context.GetAllExtenders(), GeneratedMenu->StyleSet);
+		PopulateMenuBarBuilder(MenuBarBuilder, GeneratedMenu);
+		TSharedRef<SWidget> Result = MenuBarBuilder.MakeWidget();
+		GeneratedMenuWidget.Widget = Result;
+		return Result;
+	}
+	else if (GeneratedMenu->MenuType == EMultiBoxType::ToolBar || GeneratedMenu->MenuType == EMultiBoxType::VerticalToolBar)
+	{
+		const EOrientation ToolBarOrientation = (GeneratedMenu->MenuType == EMultiBoxType::VerticalToolBar) ? Orient_Vertical : Orient_Horizontal;
+		FToolBarBuilder ToolbarBuilder(GeneratedMenu->Context.CommandList, GeneratedMenu->MenuName, GeneratedMenu->Context.GetAllExtenders(), ToolBarOrientation, GeneratedMenu->bToolBarForceSmallIcons);
+		ToolbarBuilder.SetIsFocusable(GeneratedMenu->bToolBarIsFocusable);
+		PopulateToolBarBuilder(ToolbarBuilder, GeneratedMenu);
+		TSharedRef<SWidget> Result = ToolbarBuilder.MakeWidget();
+		GeneratedMenuWidget.Widget = Result;
+		return Result;
+	}
+
+	return SNullWidget::NullWidget;
+}
+
+void UToolMenus::AssignSetTimerForNextTickDelegate(const FSimpleDelegate& InDelegate)
+{
+	SetTimerForNextTickDelegate = InDelegate;
+}
+
+void UToolMenus::SetNextTickTimer()
+{
+	if (!bNextTickTimerIsSet)
+	{
+		bNextTickTimerIsSet = true;
+
+		SetTimerForNextTickDelegate.ExecuteIfBound();
+	}
+}
+
+void UToolMenus::CleanupStaleWidgetsNextTick()
+{
+	bCleanupStaleWidgetsNextTick = true;
+	SetNextTickTimer();
+}
+
+void UToolMenus::RefreshAllWidgets()
+{
+	bRefreshWidgetsNextTick = true;
+	SetNextTickTimer();
+}
+
+void UToolMenus::HandleNextTick()
+{
+	if (bCleanupStaleWidgetsNextTick || bRefreshWidgetsNextTick)
+	{
+		CleanupStaleWidgets();
+		bCleanupStaleWidgetsNextTick = false;
+
+		if (bRefreshWidgetsNextTick)
+		{
+			for (auto WidgetsForMenuNameIt = GeneratedMenuWidgets.CreateIterator(); WidgetsForMenuNameIt; ++WidgetsForMenuNameIt)
+			{
+				FGeneratedToolMenuWidgets& WidgetsForMenuName = WidgetsForMenuNameIt->Value;
+				for (auto Instance = WidgetsForMenuName.Instances.CreateIterator(); Instance; ++Instance)
+				{
+					if (Instance->Widget.IsValid())
+					{
+						RefreshMenuWidget(WidgetsForMenuNameIt->Key, *Instance);
+					}
+				}
+			}
+
+			bRefreshWidgetsNextTick = false;
+		}
+	}
+
+	bNextTickTimerIsSet = false;
+}
+
+void UToolMenus::CleanupStaleWidgets()
+{
+	for (auto WidgetsForMenuNameIt = GeneratedMenuWidgets.CreateIterator(); WidgetsForMenuNameIt; ++WidgetsForMenuNameIt)
+	{
+		FGeneratedToolMenuWidgets& WidgetsForMenuName = WidgetsForMenuNameIt->Value;
+
+		for (auto Instance = WidgetsForMenuName.Instances.CreateIterator(); Instance; ++Instance)
+		{
+			if (!Instance->Widget.IsValid())
+			{
+				Instance.RemoveCurrent();
+			}
+		}
+
+		if (WidgetsForMenuName.Instances.Num() == 0)
+		{
+			WidgetsForMenuNameIt.RemoveCurrent();
+		}
+	}
+}
+
+bool UToolMenus::RefreshMenuWidget(const FName InName)
+{
+	if (FGeneratedToolMenuWidgets* WidgetsForMenuName = GeneratedMenuWidgets.Find(InName))
+	{
+		for (auto Instance = WidgetsForMenuName->Instances.CreateIterator(); Instance; ++Instance)
+		{
+			if (RefreshMenuWidget(InName, *Instance))
+			{
+				return true;
+			}
+			else
+			{
+				Instance.RemoveCurrent();
+			}
+		}
+	}
+
+	return false;
+}
+
+bool UToolMenus::RefreshMenuWidget(const FName InName, FGeneratedToolMenuWidget& GeneratedMenuWidget)
+{
+	if (!GeneratedMenuWidget.Widget.IsValid())
+	{
+		return false;
+	}
+
+	// Regenerate menu from database
+	UToolMenu* GeneratedMenu = GenerateMenu(InName, GeneratedMenuWidget.GeneratedMenu->Context);
+	GeneratedMenuWidget.GeneratedMenu = GeneratedMenu;
+
+	// Regenerate Multibox
+	TSharedRef<SMultiBoxWidget> MultiBoxWidget = StaticCastSharedRef<SMultiBoxWidget>(GeneratedMenuWidget.Widget.Pin().ToSharedRef());
+	if (GeneratedMenu->MenuType == EMultiBoxType::Menu)
+	{
+		FMenuBuilder MenuBuilder(GeneratedMenu->bShouldCloseWindowAfterMenuSelection, GeneratedMenu->Context.CommandList, GeneratedMenu->Context.GetAllExtenders(), GeneratedMenu->bCloseSelfOnly, GeneratedMenu->StyleSet, GeneratedMenu->bSearchable);
+		PopulateMenuBuilder(MenuBuilder, GeneratedMenu);
+		MultiBoxWidget->SetMultiBox(MenuBuilder.GetMultiBox());
+	}
+	else if (GeneratedMenu->MenuType == EMultiBoxType::MenuBar)
+	{
+		FMenuBarBuilder MenuBarBuilder(GeneratedMenu->Context.CommandList, GeneratedMenu->Context.GetAllExtenders(), GeneratedMenu->StyleSet);
+		PopulateMenuBarBuilder(MenuBarBuilder, GeneratedMenu);
+		MultiBoxWidget->SetMultiBox(MenuBarBuilder.GetMultiBox());
+	}
+	else if (GeneratedMenu->MenuType == EMultiBoxType::ToolBar || GeneratedMenu->MenuType == EMultiBoxType::VerticalToolBar)
+	{
+		const EOrientation ToolBarOrientation = (GeneratedMenu->MenuType == EMultiBoxType::VerticalToolBar) ? Orient_Vertical : Orient_Horizontal;
+		FToolBarBuilder ToolbarBuilder(GeneratedMenu->Context.CommandList, GeneratedMenu->MenuName, GeneratedMenu->Context.GetAllExtenders(), ToolBarOrientation, GeneratedMenu->bToolBarForceSmallIcons);
+		ToolbarBuilder.SetIsFocusable(GeneratedMenu->bToolBarIsFocusable);
+		PopulateToolBarBuilder(ToolbarBuilder, GeneratedMenu);
+		MultiBoxWidget->SetMultiBox(ToolbarBuilder.GetMultiBox());
+	}
+
+	MultiBoxWidget->BuildMultiBoxWidget();
+	return true;
+}
+
+UToolMenu* UToolMenus::GenerateMenuAsBuilder(const UToolMenu* InMenu, FToolMenuContext& InMenuContext)
+{
+	TArray<UToolMenu*> Hierarchy = CollectHierarchy(InMenu->MenuName);
+
+	// Insert InMenu as second to last so items in InMenu appear before items registered in database by other plugins
+	if (Hierarchy.Num() > 0)
+	{
+		Hierarchy.Insert((UToolMenu*)InMenu, Hierarchy.Num() - 1);
+	}
+	else
+	{
+		Hierarchy.Add((UToolMenu*)InMenu);
+	}
+
+	return GenerateMenu(Hierarchy, InMenuContext);
+}
+
+UToolMenu* UToolMenus::RegisterMenu(const FName InName, const FName InParent, EMultiBoxType InType, bool bWarnIfAlreadyRegistered)
+{
+	if (UToolMenu* Found = FindMenu(InName))
+	{
+		if (!Found->bRegistered)
+		{
+			Found->MenuParent = InParent;
+			Found->MenuType = InType;
+			Found->MenuOwner = CurrentOwner();
+			Found->bRegistered = true;
+		}
+		else if (bWarnIfAlreadyRegistered)
+		{
+			UE_LOG(LogToolMenus, Warning, TEXT("Menu already registered : %s"), *InName.ToString());
+		}
+
+		return Found;
+	}
+
+	UToolMenu* ToolMenu = NewObject<UToolMenu>(this);
+	ToolMenu->InitMenu(CurrentOwner(), InName, InParent, InType);
+	ToolMenu->bRegistered = true;
+	Menus.Add(InName, ToolMenu);
+	return ToolMenu;
+}
+
+UToolMenu* UToolMenus::ExtendMenu(const FName InName)
+{
+	if (UToolMenu* Found = FindMenu(InName))
+	{
+		return Found;
+	}
+
+	UToolMenu* ToolMenu = NewObject<UToolMenu>(this);
+	ToolMenu->MenuName = InName;
+	ToolMenu->bRegistered = false;
+	Menus.Add(InName, ToolMenu);
+	return ToolMenu;
+}
+
+void UToolMenus::RemoveMenu(const FName MenuName)
+{
+	Menus.Remove(MenuName);
+}
+
+bool UToolMenus::AddMenuEntryObject(UToolMenuEntryScript* MenuEntryObject)
+{
+	UToolMenu* Menu = UToolMenus::Get()->ExtendMenu(MenuEntryObject->Data.Menu);
+	Menu->AddMenuEntryObject(MenuEntryObject);
+	return true;
+}
+
+void UToolMenus::SetSectionLabel(const FName MenuName, const FName SectionName, const FText Label)
+{
+	ExtendMenu(MenuName)->FindOrAddSection(SectionName).Label = TAttribute<FText>(Label);
+}
+
+void UToolMenus::SetSectionPosition(const FName MenuName, const FName SectionName, const FName PositionName, const EToolMenuInsertType PositionType)
+{
+	ExtendMenu(MenuName)->FindOrAddSection(SectionName).InsertPosition = FToolMenuInsert(PositionName, PositionType);
+}
+
+void UToolMenus::AddSection(const FName MenuName, const FName SectionName, const TAttribute< FText >& InLabel, const FToolMenuInsert InPosition)
+{
+	UToolMenu* Menu = ExtendMenu(MenuName);
+	FToolMenuSection* Section = Menu->FindSection(SectionName);
+	if (!Section)
+	{
+		Menu->AddSection(SectionName, InLabel, InPosition);
+	}
+}
+
+void UToolMenus::RemoveSection(const FName MenuName, const FName InSection)
+{
+	if (UToolMenu* Menu = FindMenu(MenuName))
+	{
+		Menu->RemoveSection(InSection);
+	}
+}
+
+void UToolMenus::AddEntry(const FName MenuName, const FName InSection, const FToolMenuEntry& InEntry)
+{
+	ExtendMenu(MenuName)->FindOrAddSection(InSection).AddEntry(InEntry);
+}
+
+void UToolMenus::RemoveEntry(const FName MenuName, const FName InSection, const FName InName)
+{
+	if (UToolMenu* Menu = FindMenu(MenuName))
+	{
+		if (FToolMenuSection* Section = Menu->FindSection(InSection))
+		{
+			Section->RemoveEntry(InName);
+		}
+	}
+}
+
+void UToolMenus::UnregisterOwnerInternal(FToolMenuOwner InOwner)
+{
+	if (InOwner != FToolMenuOwner())
+	{
+		for (auto It = Menus.CreateIterator(); It; ++It)
+		{
+			int32 NumEntriesRemoved = 0;
+
+			UToolMenu* Menu = It->Value;
+			for (FToolMenuSection& Section : Menu->Sections)
+			{
+				NumEntriesRemoved += Section.RemoveEntriesByOwner(InOwner);
+			}
+
+			// Refresh any widgets that are currently displayed to the user
+			if (NumEntriesRemoved > 0)
+			{
+				RefreshAllWidgets();
+			}
+		}
+	}
+}
+
+FToolMenuOwner UToolMenus::CurrentOwner() const
+{
+	if (OwnerStack.Num() > 0)
+	{
+		return OwnerStack.Last();
+	}
+
+	return FToolMenuOwner();
+}
+
+void UToolMenus::PushOwner(const FToolMenuOwner InOwner)
+{
+	OwnerStack.Add(InOwner);
+}
+
+void UToolMenus::PopOwner(const FToolMenuOwner InOwner)
+{
+	FToolMenuOwner PoppedOwner = OwnerStack.Pop(false);
+	check(PoppedOwner == InOwner);
+}
+
+void UToolMenus::UnregisterOwnerByName(FName InOwnerName)
+{
+	UnregisterOwnerInternal(InOwnerName);
+}
+
+void UToolMenus::RegisterStringCommandHandler(const FName InName, const FToolMenuExecuteString& InDelegate)
+{
+	StringCommandHandlers.Add(InName, InDelegate);
+}
+
+void UToolMenus::UnregisterStringCommandHandler(const FName InName)
+{
+	StringCommandHandlers.Remove(InName);
+}
diff --git a/Engine/Source/Developer/ToolMenus/Private/ToolMenusBlueprintLibrary.cpp b/Engine/Source/Developer/ToolMenus/Private/ToolMenusBlueprintLibrary.cpp
new file mode 100644
--- /dev/null
+++ b/Engine/Source/Developer/ToolMenus/Private/ToolMenusBlueprintLibrary.cpp
@@ -0,0 +1,122 @@
+// Copyright 1998-2019 Epic Games, Inc. All Rights Reserved.
+
+
+#include "ToolMenusBlueprintLibrary.h"
+
+FScriptSlateIcon UToolMenuEntryExtensions::MakeScriptSlateIcon(const FName StyleSetName, const FName StyleName, const FName SmallStyleName)
+{
+	if (SmallStyleName == NAME_None)
+	{
+		return FScriptSlateIcon(StyleSetName, StyleName);
+	}
+	else
+	{
+		return FScriptSlateIcon(StyleSetName, StyleName, SmallStyleName);
+	}
+}
+
+void UToolMenuEntryExtensions::BreakScriptSlateIcon(const FScriptSlateIcon& InValue, FName& StyleSetName, FName& StyleName, FName& SmallStyleName)
+{
+	StyleSetName = InValue.StyleSetName;
+	StyleName = InValue.StyleName;
+	SmallStyleName = InValue.SmallStyleName;
+}
+
+FToolMenuStringCommand UToolMenuEntryExtensions::MakeStringCommand(EToolMenuStringCommandType Type, FName CustomType, const FString& String)
+{
+	return FToolMenuStringCommand(Type, CustomType, String);
+}
+
+void UToolMenuEntryExtensions::BreakStringCommand(const FToolMenuStringCommand& InValue, EToolMenuStringCommandType& Type, FName& CustomType, FString& String)
+{
+	Type = InValue.Type;
+	CustomType = InValue.CustomType;
+	String = InValue.String;
+}
+
+FToolMenuOwner UToolMenuEntryExtensions::MakeToolMenuOwner(FName Name)
+{
+	return FToolMenuOwner(Name);
+}
+
+void UToolMenuEntryExtensions::BreakToolMenuOwner(const FToolMenuOwner& InValue, FName& Name)
+{
+	Name = InValue.TryGetName();
+}
+
+UObject* UToolMenuContextExtensions::FindByClass(const FToolMenuContext& Context, UClass* InClass)
+{
+	return Context.FindByClass(InClass);
+}
+
+void UToolMenuEntryExtensions::SetLabel(UPARAM(ref) FToolMenuEntry& Target, const FText& Label)
+{
+	Target.Label = Label;
+}
+
+FText UToolMenuEntryExtensions::GetLabel(const FToolMenuEntry& Target)
+{
+	return Target.Label.Get();
+}
+
+void UToolMenuEntryExtensions::SetToolTip(UPARAM(ref) FToolMenuEntry& Target, const FText& ToolTip)
+{
+	Target.ToolTip = ToolTip;
+}
+
+FText UToolMenuEntryExtensions::GetToolTip(const FToolMenuEntry& Target)
+{
+	return Target.ToolTip.Get();
+}
+
+void UToolMenuEntryExtensions::SetIcon(UPARAM(ref) FToolMenuEntry& Target, const FName StyleSetName, const FName StyleName, const FName SmallStyleName)
+{
+	if (SmallStyleName == NAME_None)
+	{
+		if (StyleSetName == NAME_None && StyleName == NAME_None)
+		{
+			Target.Icon = FSlateIcon();
+		}
+			
+		Target.Icon = FSlateIcon(StyleSetName, StyleName);
+	}
+
+	Target.Icon = FSlateIcon(StyleSetName, StyleName, SmallStyleName);
+}
+
+void UToolMenuEntryExtensions::SetStringCommand(UPARAM(ref) FToolMenuEntry& Target, const EToolMenuStringCommandType Type, const FName CustomType, const FString& String)
+{
+	Target.ResetActions();
+	Target.StringExecuteAction.Type = Type;
+	Target.StringExecuteAction.CustomType = CustomType;
+	Target.StringExecuteAction.String = String;
+}
+
+FToolMenuEntry UToolMenuEntryExtensions::InitMenuEntry(const FName InOwner, const FName InName, const FText& InLabel, const FText& InToolTip, const EToolMenuStringCommandType CommandType, const FName CustomCommandType, const FString& CommandString)
+{
+	FToolMenuEntry Entry(InOwner, InName, EMultiBlockType::MenuEntry);
+	Entry.Label = InLabel;
+	Entry.ToolTip = InToolTip;
+	Entry.StringExecuteAction = FToolMenuStringCommand(CommandType, CustomCommandType, CommandString);
+	return Entry;
+}
+
+void UToolMenuSectionExtensions::SetLabel(UPARAM(ref) FToolMenuSection& Section, const FText& Label)
+{
+	Section.Label = Label;
+}
+
+FText UToolMenuSectionExtensions::GetLabel(const FToolMenuSection& Section)
+{
+	return Section.Label.Get();
+}
+
+void UToolMenuSectionExtensions::AddEntry(UPARAM(ref) FToolMenuSection& Section, const FToolMenuEntry& Args)
+{
+	Section.AddEntry(Args);
+}
+
+void UToolMenuSectionExtensions::AddEntryObject(UPARAM(ref) FToolMenuSection& Section, UToolMenuEntryScript* InObject)
+{
+	Section.AddEntryObject(InObject);
+}
diff --git a/Engine/Source/Developer/ToolMenus/Private/ToolMenusBlueprintLibrary.h b/Engine/Source/Developer/ToolMenus/Private/ToolMenusBlueprintLibrary.h
new file mode 100644
--- /dev/null
+++ b/Engine/Source/Developer/ToolMenus/Private/ToolMenusBlueprintLibrary.h
@@ -0,0 +1,93 @@
+// Copyright 1998-2019 Epic Games, Inc. All Rights Reserved.
+
+#pragma once
+
+#include "CoreMinimal.h"
+
+#include "ToolMenuDelegates.h"
+#include "ToolMenuEntry.h"
+#include "ToolMenuSection.h"
+#include "ToolMenuMisc.h"
+#include "ToolMenuEntryScript.h"
+#include "Misc/Attribute.h"
+#include "Kismet/BlueprintFunctionLibrary.h"
+
+#include "ToolMenusBlueprintLibrary.generated.h"
+
+UCLASS()
+class UToolMenuContextExtensions : public UBlueprintFunctionLibrary
+{
+	GENERATED_BODY()
+
+public:
+
+	UFUNCTION(BlueprintCallable, Category = "Tool Menus", meta = (ScriptMethod))
+	static UObject* FindByClass(const FToolMenuContext& Context, UClass* InClass);
+};
+
+UCLASS()
+class UToolMenuEntryExtensions : public UBlueprintFunctionLibrary
+{
+	GENERATED_BODY()
+
+public:
+
+	UFUNCTION(BlueprintPure, Category = "Tool Menus", meta = (Keywords = "construct build", NativeMakeFunc))
+	static FScriptSlateIcon MakeScriptSlateIcon(const FName StyleSetName, const FName StyleName, const FName SmallStyleName = NAME_None);
+
+	UFUNCTION(BlueprintPure, Category = "Tool Menus", meta = (NativeBreakFunc))
+	static void BreakScriptSlateIcon(const FScriptSlateIcon& InValue, FName& StyleSetName, FName& StyleName, FName& SmallStyleName);
+
+	UFUNCTION(BlueprintPure, Category = "Tool Menus", meta = (Keywords = "construct build", NativeMakeFunc))
+	static FToolMenuStringCommand MakeStringCommand(EToolMenuStringCommandType Type, FName CustomType, const FString& String);
+
+	UFUNCTION(BlueprintPure, Category = "Tool Menus", meta = (NativeBreakFunc))
+	static void BreakStringCommand(const FToolMenuStringCommand& InValue, EToolMenuStringCommandType& Type, FName& CustomType, FString& String);
+
+	UFUNCTION(BlueprintPure, Category = "Tool Menus", meta = (Keywords = "construct build", NativeMakeFunc))
+	static FToolMenuOwner MakeToolMenuOwner(FName Name);
+
+	UFUNCTION(BlueprintPure, Category = "Tool Menus", meta = (NativeBreakFunc))
+	static void BreakToolMenuOwner(const FToolMenuOwner& InValue, FName& Name);
+
+	UFUNCTION(BlueprintCallable, Category = "Tool Menus", meta = (ScriptMethod))
+	static void SetLabel(UPARAM(ref) FToolMenuEntry& Target, const FText& Label);
+
+	UFUNCTION(BlueprintCallable, Category = "Tool Menus", meta = (ScriptMethod))
+	static FText GetLabel(const FToolMenuEntry& Target);
+
+	UFUNCTION(BlueprintCallable, Category = "Tool Menus", meta = (ScriptMethod))
+	static void SetToolTip(UPARAM(ref) FToolMenuEntry& Target, const FText& ToolTip);
+
+	UFUNCTION(BlueprintCallable, Category = "Tool Menus", meta = (ScriptMethod))
+	static FText GetToolTip(const FToolMenuEntry& Target);
+
+	UFUNCTION(BlueprintCallable, Category = "Tool Menus", meta = (ScriptMethod))
+	static void SetIcon(UPARAM(ref) FToolMenuEntry& Target, const FName StyleSetName, const FName StyleName = NAME_None, const FName SmallStyleName = NAME_None);
+
+	UFUNCTION(BlueprintCallable, Category = "Tool Menus", meta = (ScriptMethod))
+	static void SetStringCommand(UPARAM(ref) FToolMenuEntry& Target, const EToolMenuStringCommandType Type, const FName CustomType, const FString& String);
+
+	UFUNCTION(BlueprintCallable, Category = "Tool Menus")
+	static FToolMenuEntry InitMenuEntry(const FName InOwner, const FName InName, const FText& InLabel, const FText& InToolTip, const EToolMenuStringCommandType CommandType, const FName CustomCommandType, const FString& CommandString);
+};
+
+UCLASS()
+class UToolMenuSectionExtensions : public UBlueprintFunctionLibrary
+{
+	GENERATED_BODY()
+
+public:
+
+	UFUNCTION(BlueprintCallable, Category = "Tool Menus", meta = (ScriptMethod))
+	static void SetLabel(UPARAM(ref) FToolMenuSection& Section, const FText& Label);
+
+	UFUNCTION(BlueprintCallable, Category = "Tool Menus", meta = (ScriptMethod))
+	static FText GetLabel(const FToolMenuSection& Section);
+
+	UFUNCTION(BlueprintCallable, Category = "Tool Menus", meta = (ScriptMethod))
+	static void AddEntry(UPARAM(ref) FToolMenuSection& Section, const FToolMenuEntry& Args);
+
+	UFUNCTION(BlueprintCallable, Category = "Tool Menus", meta = (ScriptMethod))
+	static void AddEntryObject(UPARAM(ref) FToolMenuSection& Section, UToolMenuEntryScript* InObject);
+};
diff --git a/Engine/Source/Developer/ToolMenus/Private/ToolMenusModule.cpp b/Engine/Source/Developer/ToolMenus/Private/ToolMenusModule.cpp
new file mode 100644
--- /dev/null
+++ b/Engine/Source/Developer/ToolMenus/Private/ToolMenusModule.cpp
@@ -0,0 +1,30 @@
+// Copyright 1998-2019 Epic Games, Inc. All Rights Reserved.
+
+#include "CoreMinimal.h"
+#include "Modules/ModuleManager.h"
+#include "IToolMenusModule.h"
+
+
+/**
+ * Implements the Tool menus module.
+ */
+class FToolMenusModule
+	: public IToolMenusModule
+{
+public:
+
+	// IToolMenusModule interface
+
+	virtual void StartupModule( ) override
+	{
+	}
+
+	virtual void ShutdownModule( ) override
+	{
+	}
+
+	// End IToolInterface interface
+};
+
+
+IMPLEMENT_MODULE(FToolMenusModule, ToolMenus)
diff --git a/Engine/Source/Developer/ToolMenus/Public/IToolMenusModule.h b/Engine/Source/Developer/ToolMenus/Public/IToolMenusModule.h
new file mode 100644
--- /dev/null
+++ b/Engine/Source/Developer/ToolMenus/Public/IToolMenusModule.h
@@ -0,0 +1,32 @@
+// Copyright 1998-2019 Epic Games, Inc. All Rights Reserved.
+#pragma once
+
+#include "CoreMinimal.h"
+#include "Modules/ModuleInterface.h"
+#include "Modules/ModuleManager.h"
+
+DEFINE_LOG_CATEGORY_STATIC(LogToolMenus, Log, All);
+
+class IToolMenusModule : public IModuleInterface
+{
+public:
+
+	/**
+	 * Retrieve the module instance.
+	 */
+	static inline IToolMenusModule& Get()
+	{
+		return FModuleManager::LoadModuleChecked<IToolMenusModule>("ToolMenus");
+	}
+
+	/**
+	 * Checks to see if this module is loaded and ready.  It is only valid to call Get() if IsAvailable() returns true.
+	 *
+	 * @return True if the module is loaded and ready to use
+	 */
+	static inline bool IsAvailable()
+	{
+		return FModuleManager::Get().IsModuleLoaded("ToolMenus");
+	}
+};
+
diff --git a/Engine/Source/Developer/ToolMenus/Public/ToolMenu.h b/Engine/Source/Developer/ToolMenus/Public/ToolMenu.h
new file mode 100644
--- /dev/null
+++ b/Engine/Source/Developer/ToolMenus/Public/ToolMenu.h
@@ -0,0 +1,126 @@
+// Copyright 1998-2019 Epic Games, Inc. All Rights Reserved.
+
+#pragma once
+
+#include "CoreMinimal.h"
+
+#include "ToolMenuOwner.h"
+#include "ToolMenuDelegates.h"
+#include "ToolMenuSection.h"
+#include "ToolMenuContext.h"
+#include "UObject/Object.h"
+
+#include "ToolMenu.generated.h"
+
+UCLASS(BlueprintType)
+class TOOLMENUS_API UToolMenu : public UObject
+{
+	GENERATED_BODY()
+
+public:
+
+	UToolMenu();
+
+	UFUNCTION(BlueprintCallable, Category = "Tool Menus")
+	void InitMenu(const FToolMenuOwner Owner, FName Name, FName Parent = NAME_None, EMultiBoxType Type = EMultiBoxType::Menu);
+
+	UFUNCTION(BlueprintCallable, Category = "Tool Menus", meta = ( ScriptName = "AddSection", DisplayName = "Add Section", AutoCreateRefTerm = "Label", AdvancedDisplay = "InsertName,InsertType" ))
+	void AddSectionScript(const FName SectionName, const FText& Label = FText(), const FName InsertName = NAME_None, const EToolMenuInsertType InsertType = EToolMenuInsertType::Default);
+
+	UFUNCTION(BlueprintCallable, Category = "Tool Menus", meta = ( ScriptName = "AddDynamicSection", DisplayName = "Add Dynamic Section" ))
+	void AddDynamicSectionScript(const FName SectionName, UToolMenuSectionDynamic* Object);
+
+	UFUNCTION(BlueprintCallable, Category = "Tool Menus")
+	void AddMenuEntry(const FName SectionName, const FToolMenuEntry& Args);
+
+	UFUNCTION(BlueprintCallable, Category = "Tool Menus")
+	void AddMenuEntryObject(UToolMenuEntryScript* InObject);
+
+	UFUNCTION(BlueprintCallable, Category = "Tool Menus", meta = ( ScriptName = "AddSubMenu", AutoCreateRefTerm = "Label,ToolTip" ))
+	UToolMenu* AddSubMenuScript(const FName Owner, const FName SectionName, const FName Name, const FText& Label, const FText& ToolTip = FText());
+
+	UToolMenu* AddSubMenu(const FToolMenuOwner Owner, const FName SectionName, const FName Name, const FText& Label, const FText& ToolTip = FText());
+
+	void RemoveSection(const FName SectionName);
+
+	FToolMenuSection& AddSection(const FName SectionName, const TAttribute< FText >& InLabel = TAttribute<FText>(), const FToolMenuInsert InPosition = FToolMenuInsert());
+
+	FToolMenuSection& AddDynamicSection(const FName SectionName, const FNewSectionConstructChoice& InConstruct, const FToolMenuInsert InPosition = FToolMenuInsert());
+
+	FToolMenuSection* FindSection(const FName SectionName);
+
+	FToolMenuSection& FindOrAddSection(const FName SectionName);
+
+	FName GetMenuName() const { return MenuName; }
+
+	bool IsRegistered() const { return bRegistered; }
+
+	template <typename TContextType>
+	TContextType* FindContext() const
+	{
+		return Context.Find<TContextType>();
+	}
+
+	//~ Begin UObject Interface
+	virtual bool IsDestructionThreadSafe() const { return false; }
+	//~ End UObject Interface
+
+	friend class UToolMenus;
+
+private:
+
+	void InitGeneratedCopy(const UToolMenu* Source);
+
+	bool FindEntry(const FName EntryName, int32& SectionIndex, int32& EntryIndex) const;
+
+	int32 IndexOfSection(const FName SectionName) const;
+
+	int32 FindInsertIndex(const FToolMenuSection& InSection) const;
+
+public:
+
+	UPROPERTY(EditAnywhere, BlueprintReadWrite, Category = "Tool Menus")
+	FName MenuName;
+
+	UPROPERTY(EditAnywhere, BlueprintReadWrite, Category = "Tool Menus")
+	FName MenuParent;
+
+	UPROPERTY(EditAnywhere, BlueprintReadWrite, Category = "Tool Menus")
+	FName StyleName;
+
+	UPROPERTY(EditAnywhere, BlueprintReadWrite, Category = "Tool Menus")
+	FName TutorialHighlightName;
+
+	UPROPERTY(EditAnywhere, BlueprintReadWrite, Category = "Tool Menus")
+	EMultiBoxType MenuType;
+
+	UPROPERTY(EditAnywhere, BlueprintReadWrite, Category = "Tool Menus")
+	bool bShouldCloseWindowAfterMenuSelection;
+
+	UPROPERTY(EditAnywhere, BlueprintReadWrite, Category = "Tool Menus")
+	bool bCloseSelfOnly;
+
+	UPROPERTY(EditAnywhere, BlueprintReadWrite, Category = "Tool Menus")
+	bool bSearchable;
+
+	UPROPERTY(EditAnywhere, BlueprintReadWrite, Category = "ToolBar")
+	bool bToolBarIsFocusable;
+
+	UPROPERTY(EditAnywhere, BlueprintReadWrite, Category = "ToolBar")
+	bool bToolBarForceSmallIcons;
+
+	UPROPERTY(EditAnywhere, BlueprintReadWrite, Category = "Tool Menus")
+	FToolMenuOwner MenuOwner;
+	
+	UPROPERTY()
+	FToolMenuContext Context;
+
+	UPROPERTY()
+	TArray<FToolMenuSection> Sections;
+
+private:
+
+	bool bRegistered;
+
+	const ISlateStyle* StyleSet;
+};
diff --git a/Engine/Source/Developer/ToolMenus/Public/ToolMenuContext.h b/Engine/Source/Developer/ToolMenus/Public/ToolMenuContext.h
new file mode 100644
--- /dev/null
+++ b/Engine/Source/Developer/ToolMenus/Public/ToolMenuContext.h
@@ -0,0 +1,82 @@
+// Copyright 1998-2019 Epic Games, Inc. All Rights Reserved.
+
+#pragma once
+
+#include "CoreMinimal.h"
+
+#include "ToolMenuContext.generated.h"
+
+struct FUIAction;
+class FUICommandInfo;
+class FUICommandList;
+class FTabManager;
+class FExtender;
+
+UCLASS(BlueprintType, Abstract)
+class TOOLMENUS_API UToolMenuContextBase : public UObject
+{
+	GENERATED_BODY()
+};
+
+UCLASS()
+class TOOLMENUS_API USlateTabManagerContext : public UToolMenuContextBase
+{
+	GENERATED_BODY()
+public:
+
+	TWeakPtr<FTabManager> TabManager;
+};
+
+USTRUCT(BlueprintType)
+struct TOOLMENUS_API FToolMenuContext
+{
+	GENERATED_BODY()
+public:
+
+	FToolMenuContext();
+	FToolMenuContext(UObject* InContext);
+	FToolMenuContext(TSharedPtr<FUICommandList> InCommandList, TSharedPtr<FExtender> InExtender = TSharedPtr<FExtender>(), UObject* InContext = nullptr);
+
+	template <typename TContextType>
+	TContextType* Find() const
+	{
+		for (UObject* Object : ContextObjects)
+		{
+			if (TContextType* Result = Cast<TContextType>(Object))
+			{
+				return Result;
+			}
+		}
+
+		return nullptr;
+	}
+
+	UObject* FindByClass(UClass* InClass) const;
+
+	void AppendCommandList(const TSharedRef<FUICommandList>& InCommandList);
+	void AppendCommandList(const TSharedPtr<FUICommandList>& InCommandList);
+	const FUIAction* GetActionForCommand(TSharedPtr<const FUICommandInfo> Command, TSharedPtr<const FUICommandList>& OutCommandList) const;
+
+	void AddExtender(const TSharedPtr<FExtender>& InExtender);
+	TSharedPtr<FExtender> GetAllExtenders();
+	void ReplaceExtenders(const TSharedPtr<FExtender>& InExtender);
+	void ResetExtenders();
+
+	void AppendObjects(const TArray<UObject*>& InObjects);
+	void AddObject(UObject* InObject);
+
+	friend class UToolMenus;
+	friend struct FToolMenuEntry;
+
+private:
+
+	UPROPERTY()
+	TArray<UObject*> ContextObjects;
+
+	TArray<TSharedPtr<FUICommandList>> CommandLists;
+
+	TSharedPtr<FUICommandList> CommandList;
+
+	TArray<TSharedPtr<FExtender>> Extenders;
+};
+
diff --git a/Engine/Source/Developer/ToolMenus/Public/ToolMenuDelegates.h b/Engine/Source/Developer/ToolMenus/Public/ToolMenuDelegates.h
new file mode 100644
--- /dev/null
+++ b/Engine/Source/Developer/ToolMenus/Public/ToolMenuDelegates.h
@@ -0,0 +1,136 @@
+// Copyright 1998-2019 Epic Games, Inc. All Rights Reserved.
+
+#pragma once
+
+#include "CoreMinimal.h"
+
+#include "ToolMenuContext.h"
+#include "ToolMenuMisc.h"
+
+#include "Framework/Commands/UIAction.h"
+#include "Framework/Commands/UICommandList.h"
+#include "Framework/MultiBox/MultiBoxBuilder.h"
+
+#include "ToolMenuDelegates.generated.h"
+
+class UToolMenu;
+struct FToolMenuSection;
+
+
+DECLARE_DELEGATE_OneParam(FNewToolMenuSectionDelegate, FToolMenuSection&);
+DECLARE_DELEGATE_OneParam(FNewToolMenuDelegate, UToolMenu*);
+DECLARE_DELEGATE_TwoParams(FNewToolMenuDelegateLegacy, class FMenuBuilder&, UToolMenu*);
+DECLARE_DELEGATE_TwoParams(FNewToolBarDelegateLegacy, class FToolBarBuilder&, UToolMenu*);
+DECLARE_DELEGATE_RetVal_OneParam(TSharedRef<SWidget>, FNewToolMenuWidget, const FToolMenuContext&);
+
+DECLARE_DELEGATE_OneParam(FToolMenuExecuteAction, const FToolMenuContext&);
+DECLARE_DELEGATE_RetVal_OneParam(bool, FToolMenuCanExecuteAction, const FToolMenuContext&);
+DECLARE_DELEGATE_RetVal_OneParam(bool, FToolMenuIsActionChecked, const FToolMenuContext&);
+DECLARE_DELEGATE_RetVal_OneParam(ECheckBoxState, FToolMenuGetActionCheckState, const FToolMenuContext&);
+DECLARE_DELEGATE_RetVal_OneParam(bool, FToolMenuIsActionButtonVisible, const FToolMenuContext&);
+
+DECLARE_DELEGATE_TwoParams(FToolMenuExecuteString, const FString&, const FToolMenuContext&);
+
+DECLARE_DYNAMIC_DELEGATE_OneParam(FToolMenuDynamicExecuteAction, const FToolMenuContext&, Context);
+DECLARE_DYNAMIC_DELEGATE_RetVal_OneParam(bool, FToolMenuDynamicCanExecuteAction, const FToolMenuContext&, Context);
+DECLARE_DYNAMIC_DELEGATE_RetVal_OneParam(bool, FToolMenuDynamicIsActionChecked, const FToolMenuContext&, Context);
+DECLARE_DYNAMIC_DELEGATE_RetVal_OneParam(ECheckBoxState, FToolMenuDynamicGetActionCheckState, const FToolMenuContext&, Context);
+DECLARE_DYNAMIC_DELEGATE_RetVal_OneParam(bool, FToolMenuDynamicIsActionButtonVisible, const FToolMenuContext&, Context);
+
+struct TOOLMENUS_API FToolUIAction
+{
+	FToolUIAction() {}
+	FToolUIAction(const FToolMenuExecuteAction& InExecuteAction) : ExecuteAction(InExecuteAction) {}
+
+	FToolMenuExecuteAction ExecuteAction;
+	FToolMenuCanExecuteAction CanExecuteAction;
+	FToolMenuGetActionCheckState GetActionCheckState;
+	FToolMenuIsActionButtonVisible IsActionVisibleDelegate;
+};
+
+USTRUCT(BlueprintType)
+struct TOOLMENUS_API FToolDynamicUIAction
+{
+	GENERATED_BODY()
+
+	FToolDynamicUIAction() {}
+
+	UPROPERTY(EditAnywhere, BlueprintReadWrite, Category = "Tool Menus")
+	FToolMenuDynamicExecuteAction ExecuteAction;
+
+	UPROPERTY(EditAnywhere, BlueprintReadWrite, Category = "Tool Menus")
+	FToolMenuDynamicCanExecuteAction CanExecuteAction;
+
+	UPROPERTY(EditAnywhere, BlueprintReadWrite, Category = "Tool Menus")
+	FToolMenuDynamicGetActionCheckState GetActionCheckState;
+
+	UPROPERTY(EditAnywhere, BlueprintReadWrite, Category = "Tool Menus")
+	FToolMenuDynamicIsActionButtonVisible IsActionVisibleDelegate;
+};
+
+struct TOOLMENUS_API FNewToolMenuWidgetChoice
+{
+public:
+	FNewToolMenuWidgetChoice() {}
+	FNewToolMenuWidgetChoice(const FOnGetContent& InOnGetContent) : OnGetContent(InOnGetContent) {}
+	FNewToolMenuWidgetChoice(const FNewToolMenuWidget& InNewToolMenuWidget) : NewToolMenuWidget(InNewToolMenuWidget) {}
+	FNewToolMenuWidgetChoice(const FNewToolMenuDelegate& InNewToolMenu) : NewToolMenu(InNewToolMenu) {}
+
+	FOnGetContent OnGetContent;
+	FNewToolMenuWidget NewToolMenuWidget;
+	FNewToolMenuDelegate NewToolMenu;
+};
+
+struct TOOLMENUS_API FToolUIActionChoice
+{
+public:
+	FToolUIActionChoice() {}
+	FToolUIActionChoice(const FUIAction& InAction) : Action(InAction) {}
+	FToolUIActionChoice(const FExecuteAction& InExecuteAction) : Action(InExecuteAction) {}
+	FToolUIActionChoice(const FToolUIAction& InAction) : ToolAction(InAction) {}
+	FToolUIActionChoice(const FToolDynamicUIAction& InAction) : DynamicToolAction(InAction) {}
+	FToolUIActionChoice(const FToolMenuExecuteAction& InExecuteAction) : ToolAction(InExecuteAction) {}
+	FToolUIActionChoice(const TSharedPtr< const FUICommandInfo >& InCommand, const FUICommandList& InCommandList);
+
+	const FUIAction* GetUIAction() const
+	{
+		return Action.IsSet() ? &Action.GetValue() : nullptr;
+	}
+
+	const FToolUIAction* GetToolUIAction() const
+	{
+		return ToolAction.IsSet() ? &ToolAction.GetValue() : nullptr;
+	}
+
+	const FToolDynamicUIAction* GetToolDynamicUIAction() const
+	{
+		return DynamicToolAction.IsSet() ? &DynamicToolAction.GetValue() : nullptr;
+	}
+
+private:
+	TOptional<FUIAction> Action;
+	TOptional<FToolUIAction> ToolAction;
+	TOptional<FToolDynamicUIAction> DynamicToolAction;
+};
+
+struct TOOLMENUS_API FNewToolMenuChoice
+{
+	FNewToolMenuChoice() {};
+	FNewToolMenuChoice(const FNewToolMenuDelegate& InNewToolMenuDelegate) : NewToolMenuDelegate(InNewToolMenuDelegate) {}
+	FNewToolMenuChoice(const FNewMenuDelegate& InNewMenuDelegate) : NewMenuDelegate(InNewMenuDelegate) {}
+
+	FNewToolMenuDelegate NewToolMenuDelegate;
+	FNewMenuDelegate NewMenuDelegate;
+};
+
+struct TOOLMENUS_API FNewSectionConstructChoice
+{
+	FNewSectionConstructChoice() {};
+	FNewSectionConstructChoice(const FNewToolMenuDelegate& InNewToolMenuDelegate) : NewToolMenuDelegate(InNewToolMenuDelegate) {}
+	FNewSectionConstructChoice(const FNewToolMenuDelegateLegacy& InNewToolMenuDelegateLegacy) : NewToolMenuDelegateLegacy(InNewToolMenuDelegateLegacy) {}
+	FNewSectionConstructChoice(const FNewToolBarDelegateLegacy& InNewToolBarDelegateLegacy) : NewToolBarDelegateLegacy(InNewToolBarDelegateLegacy) {}
+
+	FNewToolMenuDelegate NewToolMenuDelegate;
+	FNewToolMenuDelegateLegacy NewToolMenuDelegateLegacy;
+	FNewToolBarDelegateLegacy NewToolBarDelegateLegacy;
+};
diff --git a/Engine/Source/Developer/ToolMenus/Public/ToolMenuEntry.h b/Engine/Source/Developer/ToolMenus/Public/ToolMenuEntry.h
new file mode 100644
--- /dev/null
+++ b/Engine/Source/Developer/ToolMenus/Public/ToolMenuEntry.h
@@ -0,0 +1,178 @@
+// Copyright 1998-2019 Epic Games, Inc. All Rights Reserved.
+
+#pragma once
+
+#include "CoreMinimal.h"
+
+#include "ToolMenuOwner.h"
+#include "ToolMenuDelegates.h"
+#include "ToolMenuMisc.h"
+
+#include "Misc/Attribute.h"
+#include "Framework/MultiBox/MultiBoxDefs.h"
+#include "Framework/Commands/UICommandInfo.h"
+#include "UObject/TextProperty.h"
+
+#include "ToolMenuEntry.generated.h"
+
+class UToolMenuEntryScript;
+
+struct FToolMenuEntrySubMenuData
+{
+public:
+	FToolMenuEntrySubMenuData() :
+		bIsSubMenu(false),
+		bOpenSubMenuOnClick(false)
+	{
+	}
+
+	bool bIsSubMenu;
+	bool bOpenSubMenuOnClick;
+	FNewToolMenuChoice ConstructMenu;
+};
+
+struct FToolMenuEntryToolBarData
+{
+public:
+	FToolMenuEntryToolBarData() :
+		bSimpleComboBox(false),
+		bIsFocusable(false),
+		bForceSmallIcons(false)
+	{
+	}
+
+	TOptional< EVisibility > LabelVisibility;
+
+	/** If true, the icon and label won't be displayed */
+	bool bSimpleComboBox;
+
+	/** Whether ToolBar will have Focusable buttons */
+	bool bIsFocusable;
+
+	/** Whether this toolbar should always use small icons, regardless of the current settings */
+	bool bForceSmallIcons;
+
+	/** Delegate that generates a widget for this combo button's menu content.  Called when the menu is summoned. */
+	FNewToolMenuWidgetChoice ComboButtonContextMenuGenerator;
+
+	/** Legacy delegate that generates a widget for this combo button's menu content.  Called when the menu is summoned. */
+	FNewToolBarDelegateLegacy ConstructLegacy;
+};
+
+
+struct FToolMenuEntryWidgetData
+{
+public:
+	FToolMenuEntryWidgetData() :
+		bNoIndent(false),
+		bSearchable(false)
+	{
+	}
+
+	/** Remove the padding from the left of the widget that lines it up with other menu items */
+	bool bNoIndent;
+
+	/** If true, widget will be searchable */
+	bool bSearchable;
+};
+
+USTRUCT(BlueprintType)
+struct TOOLMENUS_API FToolMenuEntry
+{
+	GENERATED_BODY()
+
+	FToolMenuEntry();
+	FToolMenuEntry(const FToolMenuOwner InOwner, const FName InName, EMultiBlockType InType);
+
+	static FToolMenuEntry InitMenuEntry(const FName InName, const TAttribute<FText>& InLabel, const TAttribute<FText>& InToolTip, const TAttribute<FSlateIcon>& InIcon, const FToolUIActionChoice& InAction, const EUserInterfaceActionType UserInterfaceActionType = EUserInterfaceActionType::Button, const FName InTutorialHighlightName = NAME_None);
+	static FToolMenuEntry InitMenuEntry(const TSharedPtr< const FUICommandInfo >& InCommand, const TAttribute<FText>& InLabelOverride = TAttribute<FText>(), const TAttribute<FText>& InToolTipOverride = TAttribute<FText>(), const TAttribute<FSlateIcon>& InIconOverride = TAttribute<FSlateIcon>(), const FName InTutorialHighlightName = NAME_None, const FName InNameOverride = NAME_None);
+	static FToolMenuEntry InitMenuEntryWithCommandList(const TSharedPtr<const FUICommandInfo>& InCommand, const TSharedPtr<const FUICommandList>& InCommandList, const TAttribute<FText>& InLabelOverride = TAttribute<FText>(), const TAttribute<FText>& InToolTipOverride = TAttribute<FText>(), const TAttribute<FSlateIcon>& InIconOverride = TAttribute<FSlateIcon>(), const FName InTutorialHighlightName = NAME_None, const FName InNameOverride = NAME_None);
+	static FToolMenuEntry InitMenuEntry(const FName InName, const FToolUIActionChoice& InAction, const TSharedRef<SWidget>& Widget);
+
+	static FToolMenuEntry InitSubMenu(const FName InParentMenu, const FName InName, const TAttribute<FText>& InLabel, const TAttribute<FText>& InToolTip, const FNewToolMenuChoice& InMakeMenu, const FToolUIActionChoice& InAction, const EUserInterfaceActionType InUserInterfaceActionType, bool bInOpenSubMenuOnClick = false, const TAttribute<FSlateIcon>& InIcon = TAttribute<FSlateIcon>(), const bool bShouldCloseWindowAfterMenuSelection = true);
+	static FToolMenuEntry InitSubMenu(const FName InParentMenu, const FName InName, const TAttribute<FText>& InLabel, const TAttribute<FText>& InToolTip, const FNewToolMenuChoice& InMakeMenu, bool bInOpenSubMenuOnClick = false, const TAttribute<FSlateIcon>& InIcon = TAttribute<FSlateIcon>(), const bool bShouldCloseWindowAfterMenuSelection = true);
+	static FToolMenuEntry InitSubMenu(const FName InParentMenu, const FName InName, const FToolUIActionChoice& InAction, const TSharedRef<SWidget>& InWidget, const FNewToolMenuChoice& InMakeMenu, bool bShouldCloseWindowAfterMenuSelection = true);
+
+	static FToolMenuEntry InitToolBarButton(const FName InName, const FToolUIActionChoice& InAction, const TAttribute<FText>& InLabel = TAttribute<FText>(), const TAttribute<FText>& InToolTip = TAttribute<FText>(), const TAttribute<FSlateIcon>& InIcon = TAttribute<FSlateIcon>(), const EUserInterfaceActionType UserInterfaceActionType = EUserInterfaceActionType::Button, FName InTutorialHighlightName = NAME_None);
+	static FToolMenuEntry InitToolBarButton(const TSharedPtr< const FUICommandInfo >& InCommand, const TAttribute<FText>& InLabelOverride = TAttribute<FText>(), const TAttribute<FText>& InToolTipOverride = TAttribute<FText>(), const TAttribute<FSlateIcon>& InIconOverride = TAttribute<FSlateIcon>(), FName InTutorialHighlightName = NAME_None, const FName InNameOverride = NAME_None);
+	static FToolMenuEntry InitComboButton(const FName InName, const FToolUIActionChoice& InAction, const FNewToolMenuWidgetChoice& InMenuContentGenerator, const TAttribute<FText>& InLabelOverride = TAttribute<FText>(), const TAttribute<FText>& InToolTipOverride = TAttribute<FText>(), const TAttribute<FSlateIcon>& InIconOverride = TAttribute<FSlateIcon>(), bool bInSimpleComboBox = false, FName InTutorialHighlightName = NAME_None);
+
+	static FToolMenuEntry InitMenuSeparator(const FName InName);
+	static FToolMenuEntry InitToolBarSeparator(const FName InName);
+
+	static FToolMenuEntry InitWidget(const FName InName, const TSharedRef<SWidget>& InWidget, const FText& Label, bool bNoIndent = false, bool bSearchable = true);
+
+	bool IsSubMenu() const { return SubMenuData.bIsSubMenu; }
+
+	const FUIAction* GetActionForCommand(const FToolMenuContext& InContext, TSharedPtr<const FUICommandList>& OutCommandList) const;
+
+	void SetCommandList(const TSharedPtr<const FUICommandList>& InCommandList);
+
+	friend struct FToolMenuSection;
+	friend class UToolMenuEntryScript;
+
+private:
+
+	void SetCommand(const TSharedPtr< const FUICommandInfo >& InCommand, FName InName, const TAttribute<FText>& InLabel, const TAttribute<FText>& InToolTip, const TAttribute<FSlateIcon>& InIcon);
+
+	void ResetActions();
+
+	bool IsScriptObjectDynamicConstruct() const;
+
+	bool IsNonLegacyDynamicConstruct() const;
+
+public:
+
+	UPROPERTY(EditAnywhere, BlueprintReadWrite, Category = "Tool Menus")
+	FName Name;
+
+	UPROPERTY(EditAnywhere, BlueprintReadWrite, Category = "Tool Menus")
+	FToolMenuOwner Owner;
+
+	UPROPERTY(EditAnywhere, BlueprintReadWrite, Category = "Tool Menus")
+	EMultiBlockType Type;
+
+	UPROPERTY(EditAnywhere, BlueprintReadWrite, Category = "Tool Menus")
+	EUserInterfaceActionType UserInterfaceActionType;
+
+	UPROPERTY(EditAnywhere, BlueprintReadWrite, Category = "Tool Menus")
+	FName TutorialHighlightName;
+
+	UPROPERTY(EditAnywhere, BlueprintReadWrite, Category = "Tool Menus")
+	FToolMenuInsert InsertPosition;
+
+	UPROPERTY(EditAnywhere, BlueprintReadWrite, Category = "Tool Menus")
+	bool bShouldCloseWindowAfterMenuSelection;
+
+	UPROPERTY(EditAnywhere, BlueprintReadWrite, Category = "Tool Menus")
+	UToolMenuEntryScript* ScriptObject;
+
+	FToolMenuEntrySubMenuData SubMenuData;
+
+	FToolMenuEntryToolBarData ToolBarData;
+
+	FToolMenuEntryWidgetData WidgetData;
+
+	/** Optional delegate that returns a widget to use as this menu entry */
+	FNewToolMenuWidget MakeWidget;
+
+private:
+
+	friend class UToolMenus;
+	friend class UToolMenuEntryExtensions;
+
+	TAttribute<FText> Label;
+	TAttribute<FText> ToolTip;
+	TAttribute<FSlateIcon> Icon;
+
+	FToolUIActionChoice Action;
+
+	FToolMenuStringCommand StringExecuteAction;
+
+	TSharedPtr< const FUICommandInfo > Command;
+	TSharedPtr< const FUICommandList > CommandList;
+
+	FNewToolMenuSectionDelegate Construct;
+	FNewToolMenuDelegateLegacy ConstructLegacy;
+};
+
diff --git a/Engine/Source/Developer/ToolMenus/Public/ToolMenuEntryScript.h b/Engine/Source/Developer/ToolMenus/Public/ToolMenuEntryScript.h
new file mode 100644
--- /dev/null
+++ b/Engine/Source/Developer/ToolMenus/Public/ToolMenuEntryScript.h
@@ -0,0 +1,178 @@
+// Copyright 1998-2019 Epic Games, Inc. All Rights Reserved.
+
+#pragma once
+
+#include "CoreMinimal.h"
+
+#include "UObject/ObjectMacros.h"
+#include "UObject/UObjectGlobals.h"
+#include "UObject/TextProperty.h"
+#include "Textures/SlateIcon.h"
+#include "Framework/MultiBox/MultiBoxDefs.h"
+#include "Styling/SlateTypes.h"
+
+#include "ToolMenuMisc.h"
+#include "ToolMenuContext.h"
+#include "ToolMenuSection.h"
+
+#include "ToolMenuEntryScript.generated.h"
+
+struct FToolMenuEntry;
+struct FToolMenuSection;
+
+USTRUCT(BlueprintType, meta=(HasNativeBreak="ToolMenus.ToolMenuEntryExtensions.BreakScriptSlateIcon", HasNativeMake="ToolMenus.ToolMenuEntryExtensions.MakeScriptSlateIcon"))
+struct TOOLMENUS_API FScriptSlateIcon
+{
+	GENERATED_BODY()
+
+public:
+	FScriptSlateIcon();
+	FScriptSlateIcon(const FName InStyleSetName, const FName InStyleName);
+	FScriptSlateIcon(const FName InStyleSetName, const FName InStyleName, const FName InSmallStyleName);
+
+	UPROPERTY(EditAnywhere, BlueprintReadWrite, Category = "Tool Menus")
+	FName StyleSetName;
+
+	UPROPERTY(EditAnywhere, BlueprintReadWrite, Category = "Tool Menus")
+	FName StyleName;
+
+	UPROPERTY(EditAnywhere, BlueprintReadWrite, Category = "Tool Menus")
+	FName SmallStyleName;
+
+	operator FSlateIcon() const { return GetSlateIcon(); }
+
+	FSlateIcon GetSlateIcon() const;
+};
+
+USTRUCT(BlueprintType)
+struct TOOLMENUS_API FToolMenuEntryScriptDataAdvanced
+{
+	GENERATED_BODY()
+
+public:
+
+	FToolMenuEntryScriptDataAdvanced();
+
+	UPROPERTY(EditAnywhere, BlueprintReadWrite, Category = "Advanced")
+	FName TutorialHighlight;
+
+	UPROPERTY(EditAnywhere, BlueprintReadWrite, Category = "Advanced")
+	EMultiBlockType EntryType;
+
+	UPROPERTY(EditAnywhere, BlueprintReadWrite, Category = "Advanced")
+	EUserInterfaceActionType UserInterfaceActionType;
+
+	UPROPERTY(EditAnywhere, BlueprintReadWrite, Category = "SubMenu")
+	bool bIsSubMenu;
+
+	UPROPERTY(EditAnywhere, BlueprintReadWrite, Category = "SubMenu")
+	bool bOpenSubMenuOnClick;
+
+	UPROPERTY(EditAnywhere, BlueprintReadWrite, Category = "Advanced")
+	bool bShouldCloseWindowAfterMenuSelection;
+
+	UPROPERTY(EditAnywhere, BlueprintReadWrite, Category = "ToolBar")
+	bool bSimpleComboBox;
+};
+
+USTRUCT(BlueprintType)
+struct TOOLMENUS_API FToolMenuEntryScriptData
+{
+	GENERATED_BODY()
+
+public:
+
+	UPROPERTY(EditAnywhere, BlueprintReadWrite, Category = "Identity")
+	FName Menu;
+
+	UPROPERTY(EditAnywhere, BlueprintReadWrite, Category = "Identity")
+	FName Section;
+
+	UPROPERTY(EditAnywhere, BlueprintReadWrite, Category = "Identity")
+	FName Name;
+
+	UPROPERTY(EditAnywhere, BlueprintReadWrite, Category = "Appearance")
+	FText Label;
+
+	UPROPERTY(EditAnywhere, BlueprintReadWrite, Category = "Appearance")
+	FText ToolTip;
+
+	UPROPERTY(EditAnywhere,  BlueprintReadWrite, Category = "Appearance")
+	FScriptSlateIcon Icon;
+
+	// Optional identifier used for unregistering a group of menu items
+	UPROPERTY(EditAnywhere, BlueprintReadWrite, Category = "Advanced")
+	FName OwnerName;
+
+	UPROPERTY(EditAnywhere, BlueprintReadWrite, Category = "Identity")
+	FToolMenuInsert InsertPosition;
+
+	UPROPERTY(EditAnywhere,  BlueprintReadWrite, Category = "Advanced")
+	FToolMenuEntryScriptDataAdvanced Advanced;
+};
+
+UCLASS(Blueprintable, abstract)
+class TOOLMENUS_API UToolMenuEntryScript : public UObject
+{
+	GENERATED_BODY()
+
+public:
+
+	UFUNCTION(BlueprintImplementableEvent, Category = "Action")
+	void Execute(const FToolMenuContext& Context);
+
+	UFUNCTION(BlueprintNativeEvent, Category = "Advanced")
+	bool CanExecute(const FToolMenuContext& Context) const;
+	virtual bool CanExecute_Implementation(const FToolMenuContext& Context) const { return true; }
+
+	UFUNCTION(BlueprintNativeEvent, Category = "Advanced")
+	ECheckBoxState GetCheckState(const FToolMenuContext& Context) const;
+	virtual ECheckBoxState GetCheckState_Implementation(const FToolMenuContext& Context) const { return ECheckBoxState::Undetermined; }
+
+	UFUNCTION(BlueprintNativeEvent, Category = "Advanced")
+	bool IsVisible(const FToolMenuContext& Context) const;
+	virtual bool IsVisible_Implementation(const FToolMenuContext& Context) const { return true; }
+
+	UFUNCTION(BlueprintNativeEvent, Category = "Advanced")
+	FText GetLabel(const FToolMenuContext& Context) const;
+	virtual FText GetLabel_Implementation(const FToolMenuContext& Context) const { return Data.Label; }
+
+	UFUNCTION(BlueprintNativeEvent, Category = "Advanced")
+	FText GetToolTip(const FToolMenuContext& Context) const;
+	virtual FText GetToolTip_Implementation(const FToolMenuContext& Context) const { return Data.ToolTip; }
+
+	UFUNCTION(BlueprintNativeEvent, Category = "Advanced")
+	FScriptSlateIcon GetIcon(const FToolMenuContext& Context) const;
+	virtual FScriptSlateIcon GetIcon_Implementation(const FToolMenuContext& Context) const { return Data.Icon; }
+
+	UFUNCTION(BlueprintImplementableEvent, Category = "Advanced")
+	void ConstructMenuEntry(UToolMenu* Menu, const FName SectionName, const FToolMenuContext& Context);
+
+	UFUNCTION(BlueprintCallable, Category = "Advanced")
+	void RegisterMenuEntry();
+
+	UFUNCTION(BlueprintCallable, Category = "Advanced")
+	void InitEntry(const FName OwnerName, const FName Menu, const FName Section, const FName Name, const FText& Label = FText(), const FText& ToolTip = FText());
+
+private:
+
+	friend struct FToolMenuSection;
+	friend class UToolMenus;
+
+	TAttribute<FText> CreateLabelAttribute(FToolMenuContext& Context);
+
+	TAttribute<FText> CreateToolTipAttribute(FToolMenuContext& Context);
+
+	TAttribute<FSlateIcon> CreateIconAttribute(FToolMenuContext& Context);
+
+	void ToMenuEntry(FToolMenuEntry& Output);
+
+	bool IsDynamicConstruct() const;
+
+	FSlateIcon GetSlateIcon(const FToolMenuContext& Context) const;
+
+public:
+
+	UPROPERTY(EditAnywhere, BlueprintReadWrite, Category = "Data")
+	FToolMenuEntryScriptData Data;
+};
diff --git a/Engine/Source/Developer/ToolMenus/Public/ToolMenuMisc.h b/Engine/Source/Developer/ToolMenus/Public/ToolMenuMisc.h
new file mode 100644
--- /dev/null
+++ b/Engine/Source/Developer/ToolMenus/Public/ToolMenuMisc.h
@@ -0,0 +1,102 @@
+// Copyright 1998-2019 Epic Games, Inc. All Rights Reserved.
+
+#pragma once
+
+#include "CoreMinimal.h"
+#include "Framework/Commands/UIAction.h"
+
+#include "ToolMenuMisc.generated.h"
+
+struct FToolMenuContext;
+
+UENUM(BlueprintType)
+enum class EToolMenuStringCommandType : uint8
+{
+	Command,
+	Python,
+	Custom
+};
+
+USTRUCT(BlueprintType, meta=(HasNativeBreak="ToolMenus.ToolMenuEntryExtensions.BreakStringCommand", HasNativeMake="ToolMenus.ToolMenuEntryExtensions.MakeStringCommand"))
+struct TOOLMENUS_API FToolMenuStringCommand
+{
+	GENERATED_BODY()
+
+	FToolMenuStringCommand() : Type(EToolMenuStringCommandType::Command) {}
+
+	FToolMenuStringCommand(EToolMenuStringCommandType InType, FName InCustomType, const FString& InString) :
+		Type(InType),
+		CustomType(InCustomType),
+		String(InString)
+	{
+	}
+
+	// Which command handler to use
+	UPROPERTY(EditAnywhere, BlueprintReadWrite, Category = "Tool Menus")
+	EToolMenuStringCommandType Type;
+
+	// Which command handler to use when type is custom
+	UPROPERTY(EditAnywhere, BlueprintReadWrite, Category = "Tool Menus")
+	FName CustomType;
+
+	// String to pass to command handler
+	UPROPERTY(EditAnywhere, BlueprintReadWrite, Category = "Tool Menus")
+	FString String;
+
+private:
+
+	friend class UToolMenus;
+
+	bool IsBound() const { return String.Len() > 0; }
+
+	FExecuteAction ToExecuteAction(const FToolMenuContext& Context) const;
+
+	FName GetTypeName() const;
+};
+
+UENUM(BlueprintType)
+enum class EToolMenuInsertType : uint8
+{
+	Default,
+	Before,
+	After,
+	First
+};
+
+USTRUCT(BlueprintType)
+struct FToolMenuInsert
+{
+	GENERATED_BODY()
+
+	FToolMenuInsert() : Position(EToolMenuInsertType::Default) {}
+	FToolMenuInsert(FName InName, EToolMenuInsertType InPosition) : Name(InName), Position(InPosition) {}
+
+	// Where to insert
+	UPROPERTY(EditAnywhere, BlueprintReadWrite, Category = "Tool Menus")
+	FName Name;
+
+	// How to insert
+	UPROPERTY(EditAnywhere, BlueprintReadWrite, Category = "Tool Menus")
+	EToolMenuInsertType Position;
+
+	FORCEINLINE bool operator==(const FToolMenuInsert& Other) const
+	{
+		return Other.Name == Name && Other.Position == Position;
+	}
+
+	FORCEINLINE bool operator!=(const FToolMenuInsert& Other) const
+	{
+		return Other.Name != Name || Other.Position != Position;
+	}
+
+	bool IsDefault() const
+	{
+		return Position == EToolMenuInsertType::Default;
+	}
+
+	bool IsBeforeOrAfter() const
+	{
+		return Position == EToolMenuInsertType::Before || Position == EToolMenuInsertType::After;
+	}
+};
+
diff --git a/Engine/Source/Developer/ToolMenus/Public/ToolMenuOwner.h b/Engine/Source/Developer/ToolMenus/Public/ToolMenuOwner.h
new file mode 100644
--- /dev/null
+++ b/Engine/Source/Developer/ToolMenus/Public/ToolMenuOwner.h
@@ -0,0 +1,88 @@
+// Copyright 1998-2019 Epic Games, Inc. All Rights Reserved.
+
+#pragma once
+
+#include "CoreMinimal.h"
+
+#include "ToolMenuOwner.generated.h"
+
+
+USTRUCT(BlueprintType, meta=(HasNativeBreak="ToolMenus.ToolMenuEntryExtensions.BreakToolMenuOwner", HasNativeMake="ToolMenus.ToolMenuEntryExtensions.MakeToolMenuOwner"))
+struct TOOLMENUS_API FToolMenuOwner
+{
+	GENERATED_BODY()
+
+private:
+	struct FStoreName
+	{
+		int32 Index;
+		int32 Number;
+	};
+
+	enum class EValueType : uint8
+	{
+		None,
+		Pointer,
+		Name
+	};
+
+public:
+	FORCEINLINE FToolMenuOwner() : ValueInt64(0), ValueType(EValueType::None) {}
+	FORCEINLINE FToolMenuOwner(void* InPointer) : ValueInt64(reinterpret_cast<int64>(InPointer)), ValueType(EValueType::Pointer) {}
+
+	FToolMenuOwner(const WIDECHAR* InValue) : FToolMenuOwner(FName(InValue)) {}
+	FToolMenuOwner(const ANSICHAR* InValue) : FToolMenuOwner(FName(InValue)) {}
+
+	FToolMenuOwner(const FName InValue)
+	{
+		if (InValue == NAME_None)
+		{
+			ValueInt64 = 0;
+			ValueType = EValueType::None;
+		}
+		else
+		{
+			ValueName.Index = InValue.GetComparisonIndex().ToUnstableInt();
+			ValueName.Number = InValue.GetNumber();
+			ValueType = EValueType::Name;
+		}
+	}
+
+	FORCEINLINE bool operator==(const FToolMenuOwner& Other) const
+	{
+		return Other.ValueInt64 == ValueInt64 && Other.ValueType == ValueType;
+	}
+
+	FORCEINLINE bool operator!=(const FToolMenuOwner& Other) const
+	{
+		return Other.ValueInt64 != ValueInt64 || Other.ValueType != ValueType;
+	}
+
+	friend uint32 GetTypeHash(const FToolMenuOwner& Key)
+	{
+		return GetTypeHash(Key.ValueInt64);
+	}
+
+	FORCEINLINE bool IsSet() const { return ValueInt64 != 0; }
+
+	FName TryGetName() const
+	{
+		if (ValueType == EValueType::Name)
+		{
+			const FNameEntryId EntryId = FNameEntryId::FromUnstableInt(ValueName.Index);
+			return FName(EntryId, EntryId, ValueName.Number);
+		}
+
+		return NAME_None;
+	}
+
+private:
+
+	union
+	{
+		int64 ValueInt64;
+		FStoreName ValueName;
+	};
+
+	EValueType ValueType;
+};
diff --git a/Engine/Source/Developer/ToolMenus/Public/ToolMenuSection.h b/Engine/Source/Developer/ToolMenus/Public/ToolMenuSection.h
new file mode 100644
--- /dev/null
+++ b/Engine/Source/Developer/ToolMenus/Public/ToolMenuSection.h
@@ -0,0 +1,89 @@
+// Copyright 1998-2019 Epic Games, Inc. All Rights Reserved.
+
+#pragma once
+
+#include "CoreMinimal.h"
+
+#include "ToolMenuDelegates.h"
+#include "ToolMenuEntry.h"
+#include "Misc/Attribute.h"
+
+#include "ToolMenuSection.generated.h"
+
+UCLASS(Blueprintable)
+class TOOLMENUS_API UToolMenuSectionDynamic : public UObject
+{
+	GENERATED_BODY()
+
+public:
+
+	UFUNCTION(BlueprintImplementableEvent, Category = "Tool Menus")
+	void ConstructSections(UToolMenu* Menu, const FToolMenuContext& Context);
+};
+
+USTRUCT(BlueprintType)
+struct TOOLMENUS_API FToolMenuSection
+{
+	GENERATED_BODY()
+
+public:
+
+	FToolMenuSection();
+
+	void InitSection(const FName InName, const TAttribute< FText >& InLabel, const FToolMenuInsert InPosition);
+
+	FToolMenuEntry& AddEntryObject(UToolMenuEntryScript* InObject);
+	FToolMenuEntry& AddEntry(const FToolMenuEntry& Args);
+	FToolMenuEntry& AddMenuEntry(const FName InName, const TAttribute<FText>& InLabel, const TAttribute<FText>& InToolTip, const TAttribute<FSlateIcon>& InIcon, const FToolUIActionChoice& InAction, const EUserInterfaceActionType UserInterfaceActionType = EUserInterfaceActionType::Button, const FName InTutorialHighlightName = NAME_None);
+	FToolMenuEntry& AddMenuEntry(const TSharedPtr< const FUICommandInfo >& InCommand, const TAttribute<FText>& InLabelOverride = TAttribute<FText>(), const TAttribute<FText>& InToolTipOverride = TAttribute<FText>(), const TAttribute<FSlateIcon>& InIconOverride = TAttribute<FSlateIcon>(), const FName InTutorialHighlightName = NAME_None, const FName InNameOverride = NAME_None);
+	FToolMenuEntry& AddMenuEntryWithCommandList(const TSharedPtr<const FUICommandInfo>& InCommand, const TSharedPtr<const FUICommandList>& InCommandList, const TAttribute<FText>& InLabelOverride = TAttribute<FText>(), const TAttribute<FText>& InToolTipOverride = TAttribute<FText>(), const TAttribute<FSlateIcon>& InIconOverride = TAttribute<FSlateIcon>(), const FName InTutorialHighlightName = NAME_None, const FName InNameOverride = NAME_None);
+
+	FToolMenuEntry& AddDynamicEntry(const FName InName, const FNewToolMenuSectionDelegate& InConstruct);
+	FToolMenuEntry& AddDynamicEntry(const FName InName, const FNewToolMenuDelegateLegacy& InConstruct);
+	FToolMenuEntry& AddMenuSeparator(const FName InName);
+
+	template <typename TContextType>
+	TContextType* FindContext() const
+	{
+		return Context.Find<TContextType>();
+	}
+
+private:
+
+	void InitGeneratedSectionCopy(const FToolMenuSection& Source, FToolMenuContext& InContext);
+
+	int32 RemoveEntry(const FName InName);
+	int32 RemoveEntriesByOwner(const FToolMenuOwner InOwner);
+
+	int32 IndexOfBlock(const FName InName) const;
+	int32 FindBlockInsertIndex(const FToolMenuEntry& InBlock) const;
+
+	void AssembleBlock(const FToolMenuEntry& InBlock);
+
+	bool IsNonLegacyDynamic() const;
+
+	friend class UToolMenuSectionExtensions;
+	friend class UToolMenus;
+	friend class UToolMenu;
+
+public:
+
+	UPROPERTY(EditAnywhere, BlueprintReadWrite, Category = "Tool Menus")
+	FName Name;
+
+	UPROPERTY(EditAnywhere, BlueprintReadWrite, Category = "Tool Menus")
+	TArray<FToolMenuEntry> Blocks;
+
+	UPROPERTY(EditAnywhere, BlueprintReadWrite, Category = "Tool Menus")
+	FToolMenuInsert InsertPosition;
+
+	UPROPERTY(EditAnywhere, BlueprintReadWrite, Category = "Tool Menus")
+	FToolMenuContext Context;
+
+	UPROPERTY(EditAnywhere, BlueprintReadWrite, Category = "Tool Menus")
+	UToolMenuSectionDynamic* ToolMenuSectionDynamic;
+
+	TAttribute<FText> Label;
+
+	FNewSectionConstructChoice Construct;
+};
diff --git a/Engine/Source/Developer/ToolMenus/Public/ToolMenus.h b/Engine/Source/Developer/ToolMenus/Public/ToolMenus.h
new file mode 100644
--- /dev/null
+++ b/Engine/Source/Developer/ToolMenus/Public/ToolMenus.h
@@ -0,0 +1,399 @@
+// Copyright 1998-2019 Epic Games, Inc. All Rights Reserved.
+
+#pragma once
+
+#include "CoreMinimal.h"
+#include "UObject/ObjectMacros.h"
+#include "UObject/UObjectGlobals.h"
+#include "Templates/SubclassOf.h"
+#include "Textures/SlateIcon.h"
+#include "Misc/Attribute.h"
+
+#include "IToolMenusModule.h"
+#include "ToolMenu.h"
+#include "ToolMenuContext.h"
+#include "ToolMenuDelegates.h"
+#include "ToolMenuEntry.h"
+#include "ToolMenuEntryScript.h"
+#include "ToolMenuOwner.h"
+#include "ToolMenuSection.h"
+#include "ToolMenuMisc.h"
+
+#include "Misc/CoreDelegates.h"
+
+#include "ToolMenus.generated.h"
+
+class FMenuBuilder;
+class FMultiBox;
+class SWidget;
+
+class UToolMenu;
+class UToolMenuEntryScript;
+struct FToolMenuEntry;
+struct FToolMenuSection;
+
+USTRUCT()
+struct FCustomizedToolMenuSection
+{
+	GENERATED_BODY()
+
+	UPROPERTY()
+	FName Name;
+
+	UPROPERTY()
+	TArray<FName> Items;
+};
+
+USTRUCT()
+struct FCustomizedToolMenu
+{
+	GENERATED_BODY()
+
+	UPROPERTY()
+	FName Name;
+
+	UPROPERTY()
+	TArray<FCustomizedToolMenuSection> Sections;
+
+	UPROPERTY()
+	TArray<FName> HiddenItems;
+
+	UPROPERTY()
+	TArray<FName> HiddenSections;
+};
+
+USTRUCT()
+struct FGeneratedToolMenuWidget
+{
+	GENERATED_BODY()
+
+	UPROPERTY()
+	UToolMenu* GeneratedMenu;
+
+	TWeakPtr<SWidget> Widget;
+};
+
+USTRUCT()
+struct FGeneratedToolMenuWidgets
+{
+	GENERATED_BODY()
+
+	UPROPERTY()
+	TArray<FGeneratedToolMenuWidget> Instances;
+};
+
+UCLASS()
+class TOOLMENUS_API UToolMenus : public UObject
+{
+	GENERATED_BODY()
+
+public:
+
+	UToolMenus();
+
+	UFUNCTION(BlueprintCallable, Category = "Tool Menus")
+	static UToolMenus* Get();
+
+	/** Try to get UToolMenus without forcing ToolMenus module to load. */
+	static inline UToolMenus* TryGet()
+	{
+		if (IToolMenusModule::IsAvailable())
+		{
+			return Get();
+		}
+
+		return nullptr;
+	}
+
+	/** Unregister everything associated with the given owner without forcing ToolMenus module to load. */
+	static inline void UnregisterOwner(FToolMenuOwner Owner)
+	{
+		if (UToolMenus* ToolMenus = UToolMenus::TryGet())
+		{
+			ToolMenus->UnregisterOwnerInternal(Owner);
+		}
+	}
+
+	/**
+	 * Returns true if slate initialized and editor GUI is being used.
+	 * The application should have been initialized before this method is called.
+	 *
+	 * @return	True if slate initialized and editor GUI is being used.
+	 */
+	static bool IsToolMenuUIEnabled();
+
+	static void RegisterStartupCallback(const FSimpleMulticastDelegate::FDelegate& Delegate)
+	{
+		if (IsToolMenuUIEnabled())
+		{
+			if (UToolMenus::TryGet())
+			{
+				Delegate.Execute();
+			}
+			else
+			{
+				// Wait until UToolMenus has been initialized
+				FCoreDelegates::OnPostEngineInit.Add(Delegate);
+			}
+		}
+	}
+
+	static void UnRegisterStartupCallback(const void* UserPointer)
+	{
+		FCoreDelegates::OnPostEngineInit.RemoveAll(UserPointer);
+	}
+
+	/**
+	 * Registers a menu by name
+	 * @param	Parent	Optional name of a menu to layer on top of.
+	 * @param	Type	Type of menu that will be generated such as: ToolBar, VerticalToolBar, etc..
+	 * @param	bWarnIfAlreadyRegistered	Display warning if already registered
+	 * @return	ToolMenu	Menu object
+	 */
+	UFUNCTION(BlueprintCallable, Category = "Tool Menus")
+	UToolMenu* RegisterMenu(FName Name, const FName Parent = NAME_None, EMultiBoxType Type = EMultiBoxType::Menu, bool bWarnIfAlreadyRegistered = true);
+
+	/**
+	 * Extends a menu without registering the menu or claiming ownership of it. Ok to call even if menu does not exist yet.
+	 * @param	Name	Name of the menu to extend
+	 * @return	ToolMenu	Menu object
+	 */
+	UFUNCTION(BlueprintCallable, Category = "Tool Menus")
+	UToolMenu* ExtendMenu(const FName Name);
+
+	/**
+	 * Generate widget from a registered menu. Most common function used to generate new menu widgets.
+	 * @param	Name	Registered menu's name that widget will be generated from
+	 * @param	Context	Additional information specific to the menu being generated
+	 * @return	Widget to display
+	 */
+	TSharedRef<SWidget> GenerateWidget(const FName Name, FToolMenuContext& InMenuContext);
+
+
+
+
+	/**
+	 * Finds an existing menu that has been registered or extended.
+	 * @param	Name	Name of the menu to find.
+	 * @return	ToolMenu	Menu object. Returns null if not found.
+	 */
+	UFUNCTION(BlueprintCallable, Category = "Tool Menus")
+	UToolMenu* FindMenu(const FName Name);
+
+	/**
+	 * Determines if a menu has already been registered.
+	 * @param	Name	Name of the menu to find.
+	 * @return	bool	True if menu has already been registered.
+	 */
+	UFUNCTION(BlueprintCallable, Category = "Tool Menus")
+	bool IsMenuRegistered(const FName Name) const;
+
+	/** Rebuilds all widgets generated from a specific menu. */
+	UFUNCTION(BlueprintCallable, Category = "Tool Menus")
+	bool RefreshMenuWidget(const FName Name);
+
+	/** Rebuilds all currently generated widgets next tick. */
+	UFUNCTION(BlueprintCallable, Category = "Tool Menus")
+	void RefreshAllWidgets();
+
+	/** Registers menu entry object from blueprint/script */
+	UFUNCTION(BlueprintCallable, Category = "Tool Menus")
+	static bool AddMenuEntryObject(UToolMenuEntryScript* MenuEntryObject);
+
+	/** Removes all entries that were registered under a specific owner name */
+	UFUNCTION(BlueprintCallable, Category = "Tool Menus")
+	void UnregisterOwnerByName(FName InOwnerName);
+
+	/** Sets a section's displayed label text. */
+	UFUNCTION(BlueprintCallable, Category = "Tool Menus")
+	void SetSectionLabel(const FName MenuName, const FName SectionName, const FText Label);
+
+	/** Sets where to insert a section into a menu when generating relative to other section names. */
+	UFUNCTION(BlueprintCallable, Category = "Tool Menus")
+	void SetSectionPosition(const FName MenuName, const FName SectionName, const FName OtherSectionName, const EToolMenuInsertType PositionType);
+
+	/** Registers a section for a menu */
+	void AddSection(const FName MenuName, const FName SectionName, const TAttribute< FText >& InLabel, const FToolMenuInsert InPosition = FToolMenuInsert());
+
+	/** Registers an entry for a menu's section */
+	void AddEntry(const FName MenuName, const FName SectionName, const FToolMenuEntry& Entry);
+
+	/** Removes a menu entry from a given menu and section */
+	UFUNCTION(BlueprintCallable, Category = "Tool Menus")
+	void RemoveEntry(const FName MenuName, const FName Section, const FName Name);
+
+	/** Removes a section from a given menu */
+	UFUNCTION(BlueprintCallable, Category = "Tool Menus")
+	void RemoveSection(const FName MenuName, const FName Section);
+
+	/** Unregisters a menu by name */
+	UFUNCTION(BlueprintCallable, Category = "Tool Menus")
+	void RemoveMenu(const FName MenuName);
+
+	/** Finds a context object of a given class if it exists */
+	UFUNCTION(BlueprintCallable, Category = "Tool Menus")
+	static UObject* FindContext(const FToolMenuContext& InContext, UClass* InClass);
+
+	/**
+	 * Generate widget from a hierarchy of menus. For advanced specialized use cases.
+	 * @param	Hierarchy	Array of menus to combine into a final widget
+	 * @param	Context	Additional information specific to the menu being generated
+	 * @return	Widget to display
+	 */
+	TSharedRef<SWidget> GenerateWidget(const TArray<UToolMenu*>& Hierarchy, FToolMenuContext& InMenuContext);
+
+	/**
+	 * Generate widget from a final collapsed menu. For advanced specialized use cases.
+	 * @param	GeneratedMenu	Combined final menu to generate final widget from
+	 * @return	Widget to display
+	 */
+	TSharedRef<SWidget> GenerateWidget(UToolMenu* GeneratedMenu);
+	
+	/** Create a finalized menu that combines all parents used to generate a widget. Advanced special use cases only. */
+	UToolMenu* GenerateMenu(const FName Name, FToolMenuContext& InMenuContext);
+
+	/** Create a finalized menu that combines given hierarchy array that will generate a widget. Advanced special use cases only. */
+	UToolMenu* GenerateMenu(const TArray<UToolMenu*>& Hierarchy, FToolMenuContext& InMenuContext);
+
+	/** Create a finalized menu based on a custom crafted menu. Advanced special use cases only. */
+	UToolMenu* GenerateMenuAsBuilder(const UToolMenu* InMenu, FToolMenuContext& InMenuContext);
+
+	/** For advanced use cases */
+	void AssembleMenuByName(UToolMenu* GeneratedMenu, const FName Name);
+
+	/** For advanced use cases */
+	void AssembleMenuHierarchy(UToolMenu* GeneratedMenu, const TArray<UToolMenu*>& Hierarchy);
+
+	/** For advanced use cases */
+	FToolMenuOwner CurrentOwner() const;
+
+	/** Registers a new type of string based command handler. */
+	void RegisterStringCommandHandler(const FName InName, const FToolMenuExecuteString& InDelegate);
+
+	/** Removes a string based command handler. */
+	void UnregisterStringCommandHandler(const FName InName);
+	
+	/** Sets delegate to setup timer for deferred one off ticks */
+	void AssignSetTimerForNextTickDelegate(const FSimpleDelegate& InDelegate);
+
+	/** Timer function used to consolidate multiple duplicate requests into a single frame. */
+	void HandleNextTick();
+
+	/** Displaying extension points is for debugging menus */
+	DECLARE_DELEGATE_RetVal(bool, FShouldDisplayExtensionPoints);
+	FShouldDisplayExtensionPoints ShouldDisplayExtensionPoints;
+
+	static FName JoinMenuPaths(const FName Base, const FName Child);
+
+	friend struct FToolMenuOwnerScoped;
+	friend struct FToolMenuStringCommand;
+
+	//~ Begin UObject Interface
+	static void AddReferencedObjects(UObject* InThis, FReferenceCollector& Collector);
+	//~ End UObject Interface
+
+private:
+
+	/** Sets a timer to be called next engine tick so that multiple repeated actions can be combined together. */
+	void SetNextTickTimer();
+
+	/** Release references to UObjects of widgets that have been deleted. Combines multiple requests in one frame together for improved performance. */
+	void CleanupStaleWidgetsNextTick();
+
+	/** Release references to UObjects of widgets that have been deleted */
+	void CleanupStaleWidgets();
+
+	/** Re-creates widget that is active */
+	bool RefreshMenuWidget(const FName Name, FGeneratedToolMenuWidget& GeneratedMenuWidget);
+
+	/** Sets the current temporary menu owner to avoid needing to supply owner for each menu entry being registered. Used by FToolMenuEntryScoped */
+	void PushOwner(const FToolMenuOwner InOwner);
+
+	/** Sets the current temporary menu owner. Used by FToolMenuEntryScoped */
+	void PopOwner(const FToolMenuOwner InOwner);
+
+	UToolMenu* FindSubMenuToGenerateWith(const FName InParentName, const FName InChildName);
+
+	TArray<UToolMenu*> CollectHierarchy(const FName Name);
+
+	void FillMenu(FMenuBuilder& MenuBuilder, FName InMenuName, FToolMenuContext InMenuContext);
+	void FillMenuBarDropDown(FMenuBuilder& MenuBuilder, FName InParentName, FName InChildName, FToolMenuContext InMenuContext);
+	void PopulateMenuBuilder(FMenuBuilder& MenuBuilder, UToolMenu* MenuData);
+	void PopulateMenuBarBuilder(FMenuBarBuilder& MenuBarBuilder, UToolMenu* MenuData);
+	void PopulateToolBarBuilder(FToolBarBuilder& ToolBarBuilder, UToolMenu* MenuData);
+
+	TSharedRef<SWidget> GenerateToolbarComboButtonMenu(const FName SubMenuFullName, FToolMenuContext InContext);
+
+	FOnGetContent ConvertWidgetChoice(const FNewToolMenuWidgetChoice& Choice, const FToolMenuContext& Context) const;
+
+	/** Converts a string command to a FUIAction */
+	static FUIAction ConvertUIAction(const FToolMenuEntry& Block, const FToolMenuContext& Context);
+	static FUIAction ConvertUIAction(const FToolUIActionChoice& Choice, const FToolMenuContext& Context);
+	static FUIAction ConvertUIAction(const FToolUIAction& Actions, const FToolMenuContext& Context);
+	static FUIAction ConvertUIAction(const FToolDynamicUIAction& Actions, const FToolMenuContext& Context);
+	static FUIAction ConvertScriptObjectToUIAction(UToolMenuEntryScript* ScriptObject, const FToolMenuContext& Context);
+
+	static void ExecuteStringCommand(const FToolMenuStringCommand StringCommand, const FToolMenuContext Context);
+
+	void FillMenuDynamic(FMenuBuilder& Builder, FNewToolMenuDelegate InConstructMenu, const FToolMenuContext Context);
+
+	void ListAllParents(const FName Name, TArray<FName>& AllParents);
+
+	void AssembleMenu(UToolMenu* GeneratedMenu, const UToolMenu* Other);
+	void AssembleMenuSection(UToolMenu* GeneratedMenu, const UToolMenu* Other, FToolMenuSection* DestSection, const FToolMenuSection& OtherSection);
+
+	void CopyMenuSettings(UToolMenu* GeneratedMenu, const UToolMenu* Other);
+
+	void AddReferencedContextObjects(const TSharedRef<FMultiBox>& InMultiBox, const FToolMenuContext& InMenuContext);
+
+	void ApplyCustomization(UToolMenu* GeneratedMenu);
+
+	FCustomizedToolMenu* FindCustomizedMenu(const FName InName);
+	int32 FindCustomizedMenuIndex(const FName InName);
+
+	void UnregisterOwnerInternal(FToolMenuOwner Owner);
+
+	bool GetDisplayUIExtensionPoints() const;
+
+private:
+
+	UPROPERTY(EditAnywhere, Category = Misc)
+	TArray<FCustomizedToolMenu> CustomizedMenus;
+
+	UPROPERTY()
+	TMap<FName, UToolMenu*> Menus;
+
+	UPROPERTY()
+	TMap<FName, FGeneratedToolMenuWidgets> GeneratedMenuWidgets;
+
+	TMap<TWeakPtr<FMultiBox>, TArray<UObject*>> WidgetObjectReferences;
+
+	TArray<FToolMenuOwner> OwnerStack;
+
+	TMap<FName, FToolMenuExecuteString> StringCommandHandlers;
+
+	FSimpleDelegate SetTimerForNextTickDelegate;
+
+	bool bNextTickTimerIsSet;
+	bool bRefreshWidgetsNextTick;
+	bool bCleanupStaleWidgetsNextTick;
+};
+
+struct FToolMenuOwnerScoped
+{
+	FToolMenuOwnerScoped(const FToolMenuOwner InOwner) : Owner(InOwner)
+	{
+		UToolMenus::Get()->PushOwner(InOwner);
+	}
+
+	~FToolMenuOwnerScoped()
+	{
+		UToolMenus::Get()->PopOwner(Owner);
+	}
+
+	FToolMenuOwner GetOwner() const { return Owner; }
+
+private:
+
+	FToolMenuOwner Owner;
+};
diff --git a/Engine/Source/Developer/ToolMenus/ToolMenus.Build.cs b/Engine/Source/Developer/ToolMenus/ToolMenus.Build.cs
new file mode 100644
--- /dev/null
+++ b/Engine/Source/Developer/ToolMenus/ToolMenus.Build.cs
@@ -0,0 +1,21 @@
+// Copyright 1998-2019 Epic Games, Inc. All Rights Reserved.
+
+namespace UnrealBuildTool.Rules
+{
+	public class ToolMenus : ModuleRules
+	{
+		public ToolMenus(ReadOnlyTargetRules Target) : base(Target)
+		{
+			PrivateDependencyModuleNames.AddRange(
+				new string[] {
+					"Core",
+					"CoreUObject",
+					"ApplicationCore",
+					"Slate",
+					"SlateCore",
+					"Engine",
+				}
+			);
+		}
+	}
+}
diff --git a/Engine/Source/Editor/UnrealEd/Private/EditorEngine.cpp b/Engine/Source/Editor/UnrealEd/Private/EditorEngine.cpp
--- a/Engine/Source/Editor/UnrealEd/Private/EditorEngine.cpp
+++ b/Engine/Source/Editor/UnrealEd/Private/EditorEngine.cpp
@@ -785,9 +785,19 @@ void UEditorEngine::InitEditor(IEngineLoop* InEngineLoop)
 
 	if (FSlateApplication::IsInitialized() && UToolMenus::IsToolMenuUIEnabled())
 	{
-		UToolMenus::Get()->SetTimerManager(TimerManager.ToSharedRef());
+		TWeakPtr<FTimerManager> WeakTimerManager;
+		UToolMenus::Get()->AssignSetTimerForNextTickDelegate(FSimpleDelegate::CreateLambda([WeakTimerManager]()
+		{
+			if (WeakTimerManager.IsValid())
+			{
+				WeakTimerManager.Pin()->SetTimerForNextTick(UToolMenus::Get(), &UToolMenus::HandleNextTick);
+			}
+		}));
+
 		UToolMenus::Get()->ShouldDisplayExtensionPoints.BindStatic(&GetDisplayMultiboxHooks);
-		UToolMenus::Get()->RegisterStringCommandHandler("Command", FToolMenuExecuteString::CreateLambda([](const FString& InString, const FToolMenuContext& InContext) {
+
+		UToolMenus::Get()->RegisterStringCommandHandler("Command", FToolMenuExecuteString::CreateLambda([](const FString& InString, const FToolMenuContext& InContext)
+		{
 			GEditor->Exec(nullptr, *InString);
 		}));
 	}
diff --git a/Engine/Source/Runtime/ToolMenus/Private/ToolMenu.cpp b/Engine/Source/Runtime/ToolMenus/Private/ToolMenu.cpp
deleted file mode 100644
diff --git a/Engine/Source/Runtime/ToolMenus/Private/ToolMenuContext.cpp b/Engine/Source/Runtime/ToolMenus/Private/ToolMenuContext.cpp
deleted file mode 100644
diff --git a/Engine/Source/Runtime/ToolMenus/Private/ToolMenuEntry.cpp b/Engine/Source/Runtime/ToolMenus/Private/ToolMenuEntry.cpp
deleted file mode 100644
diff --git a/Engine/Source/Runtime/ToolMenus/Private/ToolMenuEntryScript.cpp b/Engine/Source/Runtime/ToolMenus/Private/ToolMenuEntryScript.cpp
deleted file mode 100644
diff --git a/Engine/Source/Runtime/ToolMenus/Private/ToolMenuSection.cpp b/Engine/Source/Runtime/ToolMenus/Private/ToolMenuSection.cpp
deleted file mode 100644
diff --git a/Engine/Source/Runtime/ToolMenus/Private/ToolMenus.cpp b/Engine/Source/Runtime/ToolMenus/Private/ToolMenus.cpp
deleted file mode 100644
diff --git a/Engine/Source/Runtime/ToolMenus/Private/ToolMenusBlueprintLibrary.cpp b/Engine/Source/Runtime/ToolMenus/Private/ToolMenusBlueprintLibrary.cpp
deleted file mode 100644
diff --git a/Engine/Source/Runtime/ToolMenus/Private/ToolMenusBlueprintLibrary.h b/Engine/Source/Runtime/ToolMenus/Private/ToolMenusBlueprintLibrary.h
deleted file mode 100644
diff --git a/Engine/Source/Runtime/ToolMenus/Private/ToolMenusModule.cpp b/Engine/Source/Runtime/ToolMenus/Private/ToolMenusModule.cpp
deleted file mode 100644
diff --git a/Engine/Source/Runtime/ToolMenus/Public/IToolMenusModule.h b/Engine/Source/Runtime/ToolMenus/Public/IToolMenusModule.h
deleted file mode 100644
diff --git a/Engine/Source/Runtime/ToolMenus/Public/ToolMenu.h b/Engine/Source/Runtime/ToolMenus/Public/ToolMenu.h
deleted file mode 100644
diff --git a/Engine/Source/Runtime/ToolMenus/Public/ToolMenuContext.h b/Engine/Source/Runtime/ToolMenus/Public/ToolMenuContext.h
deleted file mode 100644
diff --git a/Engine/Source/Runtime/ToolMenus/Public/ToolMenuDelegates.h b/Engine/Source/Runtime/ToolMenus/Public/ToolMenuDelegates.h
deleted file mode 100644
diff --git a/Engine/Source/Runtime/ToolMenus/Public/ToolMenuEntry.h b/Engine/Source/Runtime/ToolMenus/Public/ToolMenuEntry.h
deleted file mode 100644
diff --git a/Engine/Source/Runtime/ToolMenus/Public/ToolMenuEntryScript.h b/Engine/Source/Runtime/ToolMenus/Public/ToolMenuEntryScript.h
deleted file mode 100644
diff --git a/Engine/Source/Runtime/ToolMenus/Public/ToolMenuMisc.h b/Engine/Source/Runtime/ToolMenus/Public/ToolMenuMisc.h
deleted file mode 100644
diff --git a/Engine/Source/Runtime/ToolMenus/Public/ToolMenuOwner.h b/Engine/Source/Runtime/ToolMenus/Public/ToolMenuOwner.h
deleted file mode 100644
diff --git a/Engine/Source/Runtime/ToolMenus/Public/ToolMenuSection.h b/Engine/Source/Runtime/ToolMenus/Public/ToolMenuSection.h
deleted file mode 100644
diff --git a/Engine/Source/Runtime/ToolMenus/Public/ToolMenus.h b/Engine/Source/Runtime/ToolMenus/Public/ToolMenus.h
deleted file mode 100644
diff --git a/Engine/Source/Runtime/ToolMenus/ToolMenus.Build.cs b/Engine/Source/Runtime/ToolMenus/ToolMenus.Build.cs
deleted file mode 100644
