ToolMenu sub menus now have names to allow extending and customization
#rnx

diff --git a/Engine/Source/Developer/ToolMenus/Private/ToolMenus.cpp b/Engine/Source/Developer/ToolMenus/Private/ToolMenus.cpp
--- a/Engine/Source/Developer/ToolMenus/Private/ToolMenus.cpp
+++ b/Engine/Source/Developer/ToolMenus/Private/ToolMenus.cpp
@@ -157,15 +157,6 @@ void UToolMenus::ListAllParents(const FName InName, TArray<FName>& AllParents)
 	}
 }
 
-void UToolMenus::AssembleMenuByName(UToolMenu* GeneratedMenu, const FName InName)
-{
-	if (UToolMenu* Menu = FindMenu(InName))
-	{
-		GeneratedMenu->InitGeneratedCopy(Menu, Menu->MenuName);
-		AssembleMenuHierarchy(GeneratedMenu, CollectHierarchy(InName));
-	}
-}
-
 void UToolMenus::AssembleMenuSection(UToolMenu* GeneratedMenu, const UToolMenu* Other, FToolMenuSection* DestSection, const FToolMenuSection& OtherSection)
 {
 	// Build list of blocks in expected order including blocks created by construct delegates
@@ -510,29 +501,35 @@ void UToolMenus::AssembleMenuHierarchy(UToolMenu* GeneratedMenu, const TArray<UT
 	ApplyCustomization(GeneratedMenu);
 }
 
-void UToolMenus::FillMenuDynamic(FMenuBuilder& MenuBuilder, FNewToolMenuDelegate InConstructMenu, const FToolMenuContext Context)
+void UToolMenus::FillMenuDynamic(FMenuBuilder& MenuBuilder, const FName SubMenuFullName, FNewToolMenuDelegate InConstructMenu, const FToolMenuContext Context)
 {
-	if (InConstructMenu.IsBound())
+	TArray<UToolMenu*> Hierarchy;
+	
+	if (SubMenuFullName != NAME_None)
 	{
-		// Create final menu
-		UToolMenu* MenuData = NewObject<UToolMenu>();
-		MenuData->Context = Context;
-		InConstructMenu.Execute(MenuData);
+		Hierarchy = CollectHierarchy(SubMenuFullName);
+	}
 
-		// Populate menu builder with final menu
-		PopulateMenuBuilder(MenuBuilder, MenuData);
+	// Construct menu using delegate and insert as root so it can be overridden
+	if (InConstructMenu.IsBound())
+	{
+		UToolMenu* Menu = NewObject<UToolMenu>(this);
+		Menu->Context = Context;
+		Menu->MenuName = SubMenuFullName;
+		InConstructMenu.Execute(Menu);
+		Menu->MenuName = SubMenuFullName;
+		Hierarchy.Insert(Menu, 0);
 	}
+
+	// Populate menu builder with final menu
+	UToolMenu* GeneratedMenu = GenerateMenu(Hierarchy, Context);
+	PopulateMenuBuilder(MenuBuilder, GeneratedMenu);
 }
 
 void UToolMenus::FillMenu(class FMenuBuilder& MenuBuilder, FName InMenuName, FToolMenuContext InMenuContext)
 {
-	// Create combined final menu
-	UToolMenu* GeneratedMenu = NewObject<UToolMenu>();
-	GeneratedMenu->Context = InMenuContext;
-	AssembleMenuByName(GeneratedMenu, InMenuName);
-
 	// Populate menu builder with final menu
-	PopulateMenuBuilder(MenuBuilder, GeneratedMenu);
+	PopulateMenuBuilder(MenuBuilder, GenerateMenu(InMenuName, InMenuContext));
 }
 
 TSharedRef<SWidget> UToolMenus::GenerateToolbarComboButtonMenu(const FName SubMenuFullName, FToolMenuContext InContext)
@@ -545,9 +542,9 @@ void UToolMenus::FillMenuBarDropDown(class FMenuBuilder& MenuBuilder, FName InPa
 	if (UToolMenu* MenuToUse = FindSubMenuToGenerateWith(InParentName, InChildName))
 	{
 		// Create combined final menu
-		UToolMenu* GeneratedMenu = NewObject<UToolMenu>();
-		GeneratedMenu->Context = InMenuContext;
-		AssembleMenuByName(GeneratedMenu, MenuToUse->MenuName);
+		UToolMenu* GeneratedMenu = GenerateMenu(MenuToUse->MenuName, InMenuContext);
+
+		// Override name that people should use to extend this menu (as it may not be implemented)
 		GeneratedMenu->MenuName = JoinMenuPaths(InParentName, InChildName);
 
 		// Populate menu builder with final menu
@@ -615,10 +612,11 @@ void UToolMenus::PopulateMenuBuilder(FMenuBuilder& MenuBuilder, UToolMenu* MenuD
 					else if (Block.SubMenuData.ConstructMenu.NewToolMenuDelegate.IsBound())
 					{
 						// SubMenu constructed each time it is opened
+						FName SubMenuFullName = JoinMenuPaths(MenuData->MenuName, Block.Name);
 						MenuBuilder.AddSubMenu(
 							Block.Label,
 							Block.ToolTip,
-							FNewMenuDelegate::CreateUObject(this, &UToolMenus::FillMenuDynamic, Block.SubMenuData.ConstructMenu.NewToolMenuDelegate, MenuData->Context),
+							FNewMenuDelegate::CreateUObject(this, &UToolMenus::FillMenuDynamic, SubMenuFullName, Block.SubMenuData.ConstructMenu.NewToolMenuDelegate, MenuData->Context),
 							Block.SubMenuData.bOpenSubMenuOnClick,
 							Block.Icon.Get(),
 							Block.bShouldCloseWindowAfterMenuSelection,
@@ -1115,7 +1113,7 @@ UToolMenu* UToolMenus::GenerateMenu(const FName Name, FToolMenuContext& InMenuCo
 	return GenerateMenu(CollectHierarchy(Name), InMenuContext);
 }
 
-UToolMenu* UToolMenus::GenerateMenu(const TArray<UToolMenu*>& Hierarchy, FToolMenuContext& InMenuContext)
+UToolMenu* UToolMenus::GenerateMenu(const TArray<UToolMenu*>& Hierarchy, const FToolMenuContext& InMenuContext)
 {
 	UToolMenu* GeneratedMenu = NewObject<UToolMenu>(this);
 
diff --git a/Engine/Source/Developer/ToolMenus/Public/ToolMenus.h b/Engine/Source/Developer/ToolMenus/Public/ToolMenus.h
--- a/Engine/Source/Developer/ToolMenus/Public/ToolMenus.h
+++ b/Engine/Source/Developer/ToolMenus/Public/ToolMenus.h
@@ -253,15 +253,12 @@ public:
 	UToolMenu* GenerateMenu(const FName Name, FToolMenuContext& InMenuContext);
 
 	/** Create a finalized menu that combines given hierarchy array that will generate a widget. Advanced special use cases only. */
-	UToolMenu* GenerateMenu(const TArray<UToolMenu*>& Hierarchy, FToolMenuContext& InMenuContext);
+	UToolMenu* GenerateMenu(const TArray<UToolMenu*>& Hierarchy, const FToolMenuContext& InMenuContext);
 
 	/** Create a finalized menu based on a custom crafted menu. Advanced special use cases only. */
 	UToolMenu* GenerateMenuAsBuilder(const UToolMenu* InMenu, FToolMenuContext& InMenuContext);
 
-	/** For advanced use cases */
-	void AssembleMenuByName(UToolMenu* GeneratedMenu, const FName Name);
-
-	/** For advanced use cases */
+	/** Bake final menu including calls to construction delegates, sorting, and customization */
 	void AssembleMenuHierarchy(UToolMenu* GeneratedMenu, const TArray<UToolMenu*>& Hierarchy);
 
 	/* Returns list of menus starting with root parent */
@@ -336,7 +333,7 @@ private:
 
 	static void ExecuteStringCommand(const FToolMenuStringCommand StringCommand, const FToolMenuContext Context);
 
-	void FillMenuDynamic(FMenuBuilder& Builder, FNewToolMenuDelegate InConstructMenu, const FToolMenuContext Context);
+	void FillMenuDynamic(FMenuBuilder& Builder, const FName SubMenuFullName, FNewToolMenuDelegate InConstructMenu, const FToolMenuContext Context);
 
 	void ListAllParents(const FName Name, TArray<FName>& AllParents);
 
