Fix submenus that provide custom widgets not opening.
#rnx

diff --git a/Engine/Source/Developer/ToolMenus/Private/ToolMenus.cpp b/Engine/Source/Developer/ToolMenus/Private/ToolMenus.cpp
--- a/Engine/Source/Developer/ToolMenus/Private/ToolMenus.cpp
+++ b/Engine/Source/Developer/ToolMenus/Private/ToolMenus.cpp
@@ -597,63 +597,43 @@ void UToolMenus::PopulateMenuBuilder(FMenuBuilder& MenuBuilder, UToolMenu* MenuD
 			{
 				if (Block.IsSubMenu())
 				{
+					FName SubMenuFullName = JoinMenuPaths(MenuData->MenuName, Block.Name);
+					FNewMenuDelegate NewMenuDelegate;
 					if (Block.SubMenuData.ConstructMenu.NewMenuDelegate.IsBound())
 					{
-						MenuBuilder.AddSubMenu(
-							Block.Label,
-							Block.ToolTip,
-							Block.SubMenuData.ConstructMenu.NewMenuDelegate,
-							Block.SubMenuData.bOpenSubMenuOnClick,
-							Block.Icon.Get(),
-							Block.bShouldCloseWindowAfterMenuSelection,
-							Block.Name
-						);
+						NewMenuDelegate = Block.SubMenuData.ConstructMenu.NewMenuDelegate;
 					}
 					else if (Block.SubMenuData.ConstructMenu.NewToolMenuDelegate.IsBound())
 					{
-						// SubMenu constructed each time it is opened
-						FName SubMenuFullName = JoinMenuPaths(MenuData->MenuName, Block.Name);
-						MenuBuilder.AddSubMenu(
-							Block.Label,
-							Block.ToolTip,
-							FNewMenuDelegate::CreateUObject(this, &UToolMenus::FillMenuDynamic, SubMenuFullName, Block.SubMenuData.ConstructMenu.NewToolMenuDelegate, MenuData->Context),
-							Block.SubMenuData.bOpenSubMenuOnClick,
-							Block.Icon.Get(),
-							Block.bShouldCloseWindowAfterMenuSelection,
-							Block.Name
-						);
+						NewMenuDelegate = FNewMenuDelegate::CreateUObject(this, &UToolMenus::FillMenuDynamic, SubMenuFullName, Block.SubMenuData.ConstructMenu.NewToolMenuDelegate, MenuData->Context);
 					}
 					else
 					{
-						// SubMenu registered once by name in database
-						FName SubMenuFullName = JoinMenuPaths(MenuData->MenuName, Block.Name);
-						FNewMenuDelegate NewMenuDelegate = FNewMenuDelegate::CreateUObject(this, &UToolMenus::FillMenu, SubMenuFullName, MenuData->Context);
+						NewMenuDelegate = FNewMenuDelegate::CreateUObject(this, &UToolMenus::FillMenu, SubMenuFullName, MenuData->Context);
+					}
 
-						if (Widget.IsValid())
+					if (Widget.IsValid())
+					{
+						if (UIAction.IsBound())
 						{
-							// Could also check if Visible/Enabled bound as well
-							if (UIAction.IsBound())
-							{
-								MenuBuilder.AddSubMenu(UIAction, Widget.ToSharedRef(), NewMenuDelegate, Block.bShouldCloseWindowAfterMenuSelection);
-							}
-							else
-							{
-								MenuBuilder.AddSubMenu(Widget.ToSharedRef(), NewMenuDelegate, Block.SubMenuData.bOpenSubMenuOnClick, Block.bShouldCloseWindowAfterMenuSelection);
-							}
+							MenuBuilder.AddSubMenu(UIAction, Widget.ToSharedRef(), NewMenuDelegate, Block.bShouldCloseWindowAfterMenuSelection);
 						}
 						else
 						{
-							MenuBuilder.AddSubMenu(
-								Block.Label,
-								Block.ToolTip,
-								NewMenuDelegate,
-								Block.SubMenuData.bOpenSubMenuOnClick,
-								Block.Icon.Get(),
-								Block.bShouldCloseWindowAfterMenuSelection,
-								Block.Name
-							);
+							MenuBuilder.AddSubMenu(Widget.ToSharedRef(), NewMenuDelegate, Block.SubMenuData.bOpenSubMenuOnClick, Block.bShouldCloseWindowAfterMenuSelection);
 						}
 					}
+					else
+					{
+						MenuBuilder.AddSubMenu(
+							Block.Label,
+							Block.ToolTip,
+							NewMenuDelegate,
+							Block.SubMenuData.bOpenSubMenuOnClick,
+							Block.Icon.Get(),
+							Block.bShouldCloseWindowAfterMenuSelection,
+							Block.Name);
+					}
 				}
 				else
 				{
diff --git a/Engine/Source/Editor/LevelEditor/Private/LevelEditorCreateActorMenu.cpp b/Engine/Source/Editor/LevelEditor/Private/LevelEditorCreateActorMenu.cpp
--- a/Engine/Source/Editor/LevelEditor/Private/LevelEditorCreateActorMenu.cpp
+++ b/Engine/Source/Editor/LevelEditor/Private/LevelEditorCreateActorMenu.cpp
@@ -345,6 +345,7 @@ static void BuildSingleAssetAddReplaceActorMenu(FToolMenuSection& Section, const
 	else
 	{
 		FToolMenuEntry Entry = FToolMenuEntry::InitWidget(NAME_None, SNew( SAssetMenuEntry, Asset, AssetMenuOptions ).LabelOverride( LabelOverride ), FText());
+		Entry.Type = EMultiBlockType::MenuEntry;
 		Entry.SubMenuData.bIsSubMenu = true;
 		Entry.SubMenuData.ConstructMenu = FNewToolMenuDelegate::CreateStatic(&FillAssetAddReplaceActorMenu, Asset, AssetMenuOptions, CreateMode);
 		Section.AddEntry(Entry);
