Editor asset context menu additional customization support.
#rnx

diff --git a/Engine/Source/Developer/ToolMenus/Private/ToolMenu.cpp b/Engine/Source/Developer/ToolMenus/Private/ToolMenu.cpp
--- a/Engine/Source/Developer/ToolMenus/Private/ToolMenu.cpp
+++ b/Engine/Source/Developer/ToolMenus/Private/ToolMenu.cpp
@@ -222,3 +222,32 @@ void UToolMenu::AddMenuEntry(const FName SectionName, const FToolMenuEntry& Args
 {
 	FindOrAddSection(SectionName).AddEntry(Args);
 }
+
+bool UToolMenu::IsEditing() const
+{
+	return Context.IsEditing();
+}
+
+FName UToolMenu::GetSectionName(const FName InEntryName) const
+{
+	for (const FToolMenuSection& Section : Sections)
+	{
+		if (Section.IndexOfBlock(InEntryName) != INDEX_NONE)
+		{
+			return Section.Name;
+		}
+	}
+
+	return NAME_None;
+}
+
+FCustomizedToolMenu* UToolMenu::FindMenuCustomization() const
+{
+	return UToolMenus::Get()->FindMenuCustomization(MenuName);
+}
+
+FCustomizedToolMenu* UToolMenu::AddMenuCustomization() const
+{
+	return UToolMenus::Get()->AddMenuCustomization(MenuName);
+}
+
diff --git a/Engine/Source/Developer/ToolMenus/Private/ToolMenuContext.cpp b/Engine/Source/Developer/ToolMenus/Private/ToolMenuContext.cpp
--- a/Engine/Source/Developer/ToolMenus/Private/ToolMenuContext.cpp
+++ b/Engine/Source/Developer/ToolMenus/Private/ToolMenuContext.cpp
@@ -11,12 +11,12 @@
 #include "Internationalization/Internationalization.h"
 
 
-FToolMenuContext::FToolMenuContext()
+FToolMenuContext::FToolMenuContext() : bIsEditing(false)
 {
 
 }
 
-FToolMenuContext::FToolMenuContext(UObject* InContext)
+FToolMenuContext::FToolMenuContext(UObject* InContext) : bIsEditing(false)
 {
 	if (InContext)
 	{
@@ -24,7 +24,7 @@ FToolMenuContext::FToolMenuContext(UObject* InContext)
 	}
 }
 
-FToolMenuContext::FToolMenuContext(TSharedPtr<FUICommandList> InCommandList, TSharedPtr<FExtender> InExtender, UObject* InContext)
+FToolMenuContext::FToolMenuContext(TSharedPtr<FUICommandList> InCommandList, TSharedPtr<FExtender> InExtender, UObject* InContext) : bIsEditing(false)
 {
 	if (InContext)
 	{
diff --git a/Engine/Source/Developer/ToolMenus/Private/ToolMenuSection.cpp b/Engine/Source/Developer/ToolMenus/Private/ToolMenuSection.cpp
--- a/Engine/Source/Developer/ToolMenus/Private/ToolMenuSection.cpp
+++ b/Engine/Source/Developer/ToolMenus/Private/ToolMenuSection.cpp
@@ -113,6 +113,24 @@ FToolMenuEntry& FToolMenuSection::AddSubMenu(const FName InName, const TAttribut
 	return AddEntry(FToolMenuEntry::InitSubMenu(InName, InLabel, InToolTip, InMakeMenu, bInOpenSubMenuOnClick, InIcon, bShouldCloseWindowAfterMenuSelection));
 }
 
+FToolMenuEntry& FToolMenuSection::AddSubMenu(const FName InName, const FToolUIActionChoice& InAction, const TSharedRef<SWidget>& InWidget, const FNewToolMenuChoice& InMakeMenu, bool bShouldCloseWindowAfterMenuSelection)
+{
+	return AddEntry(FToolMenuEntry::InitSubMenu(InName, InAction, InWidget, InMakeMenu, bShouldCloseWindowAfterMenuSelection));
+}
+
+FToolMenuEntry* FToolMenuSection::FindEntry(const FName InName)
+{
+	for (int32 i=0; i < Blocks.Num(); ++i)
+	{
+		if (Blocks[i].Name == InName)
+		{
+			return &Blocks[i];
+		}
+	}
+
+	return nullptr;
+}
+
 int32 FToolMenuSection::IndexOfBlock(const FName InName) const
 {
 	for (int32 i=0; i < Blocks.Num(); ++i)
diff --git a/Engine/Source/Developer/ToolMenus/Private/ToolMenus.cpp b/Engine/Source/Developer/ToolMenus/Private/ToolMenus.cpp
--- a/Engine/Source/Developer/ToolMenus/Private/ToolMenus.cpp
+++ b/Engine/Source/Developer/ToolMenus/Private/ToolMenus.cpp
@@ -64,6 +64,26 @@ FToolUIActionChoice::FToolUIActionChoice(const TSharedPtr< const FUICommandInfo
 	}
 }
 
+void ModifyEntryForEditDialog(FToolMenuEntry& Entry)
+{
+	if (Entry.Type == EMultiBlockType::ToolBarButton)
+	{
+		Entry.Type = EMultiBlockType::MenuEntry;
+	}
+	else if (Entry.Type == EMultiBlockType::ToolBarComboButton)
+	{
+		Entry.Type = EMultiBlockType::MenuEntry;
+		if (Entry.ToolBarData.bSimpleComboBox)
+		{
+			Entry.SubMenuData.bIsSubMenu = true;
+		}
+	}
+	else if (Entry.Type == EMultiBlockType::ToolBarSeparator)
+	{
+		Entry.Type = EMultiBlockType::MenuSeparator;
+	}
+}
+
 UToolMenus::UToolMenus() :
 	bNextTickTimerIsSet(false),
 	bRefreshWidgetsNextTick(false),
@@ -375,7 +395,16 @@ void UToolMenus::AssembleMenu(UToolMenu* GeneratedMenu, const UToolMenu* Other)
 	}
 }
 
-int32 UToolMenus::FindCustomizedMenuIndex(const FName InName)
+void UToolMenus::RemoveCustomization(const FName InName)
+{
+	int32 FoundIndex = FindMenuCustomizationIndex(InName);
+	if (FoundIndex != INDEX_NONE)
+	{
+		CustomizedMenus.RemoveAt(FoundIndex, 1, false);
+	}
+}
+
+int32 UToolMenus::FindMenuCustomizationIndex(const FName InName)
 {
 	for (int32 i = 0; i < CustomizedMenus.Num(); ++i)
 	{
@@ -388,16 +417,37 @@ int32 UToolMenus::FindCustomizedMenuIndex(const FName InName)
 	return INDEX_NONE;
 }
 
-FCustomizedToolMenu* UToolMenus::FindCustomizedMenu(const FName InName)
+FCustomizedToolMenu* UToolMenus::FindMenuCustomization(const FName InName)
 {
-	int32 FoundIndex = FindCustomizedMenuIndex(InName);
-	return (FoundIndex != INDEX_NONE) ? &CustomizedMenus[FoundIndex] : nullptr;
+	for (int32 i = 0; i < CustomizedMenus.Num(); ++i)
+	{
+		if (CustomizedMenus[i].Name == InName)
+		{
+			return &CustomizedMenus[i];
+		}
+	}
+
+	return nullptr;
+}
+
+FCustomizedToolMenu* UToolMenus::AddMenuCustomization(const FName InName)
+{
+	if (FCustomizedToolMenu* Found = FindMenuCustomization(InName))
+	{
+		return Found;
+	}
+	else
+	{
+		FCustomizedToolMenu& NewCustomization = CustomizedMenus.AddDefaulted_GetRef();
+		NewCustomization.Name = InName;
+		return &NewCustomization;
+	}
 }
 
 void UToolMenus::ApplyCustomization(UToolMenu* GeneratedMenu)
 {
-	const FCustomizedToolMenu* CustomizedMenu = FindCustomizedMenu(GeneratedMenu->MenuName);
-	if (CustomizedMenu == nullptr)
+	FCustomizedToolMenu* CustomizedMenu = FindMenuCustomization(GeneratedMenu->MenuName);
+	if (!CustomizedMenu)
 	{
 		return;
 	}
@@ -418,11 +468,11 @@ void UToolMenus::ApplyCustomization(UToolMenu* GeneratedMenu)
 		TArray<FToolMenuEntry> NewBlocks;
 		NewBlocks.Reserve(Section.Blocks.Num());
 
-		for (const FName ItemName : CustomizedSection.Items)
+		for (const FCustomizedToolMenuEntry& Entry : CustomizedSection.Entries)
 		{
 			int32 EntrySectionIndex = INDEX_NONE;
 			int32 EntryIndex = INDEX_NONE;
-			if (GeneratedMenu->FindEntry(ItemName, EntrySectionIndex, EntryIndex))
+			if (GeneratedMenu->FindEntry(Entry.Name, EntrySectionIndex, EntryIndex))
 			{
 				NewBlocks.Add(GeneratedMenu->Sections[EntrySectionIndex].Blocks[EntryIndex]);
 				GeneratedMenu->Sections[EntrySectionIndex].Blocks.RemoveAt(EntryIndex);
@@ -445,22 +495,25 @@ void UToolMenus::ApplyCustomization(UToolMenu* GeneratedMenu)
 	NewSections.Append(GeneratedMenu->Sections);
 
 	// Hide
-	for (int32 SectionIndex=0; SectionIndex < NewSections.Num(); ++SectionIndex)
+	if (!GeneratedMenu->IsEditing())
 	{
-		FToolMenuSection& Section = NewSections[SectionIndex];
-		if (CustomizedMenu->HiddenSections.Contains(Section.Name))
+		for (int32 SectionIndex = 0; SectionIndex < NewSections.Num(); ++SectionIndex)
 		{
-			NewSections.RemoveAt(SectionIndex);
-			--SectionIndex;
-			continue;
-		}
+			FToolMenuSection& Section = NewSections[SectionIndex];
+			if (CustomizedMenu->HiddenSections.Contains(Section.Name))
+			{
+				NewSections.RemoveAt(SectionIndex);
+				--SectionIndex;
+				continue;
+			}
 
-		for (int32 i = 0; i < Section.Blocks.Num(); ++i)
-		{
-			if (CustomizedMenu->HiddenItems.Contains(Section.Blocks[i].Name))
+			for (int32 i = 0; i < Section.Blocks.Num(); ++i)
 			{
-				Section.Blocks.RemoveAt(i);
-				--i;
+				if (CustomizedMenu->HiddenEntries.Contains(Section.Blocks[i].Name))
+				{
+					Section.Blocks.RemoveAt(i);
+					--i;
+				}
 			}
 		}
 	}
@@ -554,7 +607,8 @@ void UToolMenus::FillMenuBarDropDown(class FMenuBuilder& MenuBuilder, FName InPa
 
 void UToolMenus::PopulateMenuBuilder(FMenuBuilder& MenuBuilder, UToolMenu* MenuData)
 {
-	if (GetDisplayUIExtensionPoints())
+	const bool bIsEditing = MenuData->IsEditing();
+	if (GetDisplayUIExtensionPoints() && !bIsEditing)
 	{
 		const FName MenuName = MenuData->GetMenuName();
 		MenuBuilder.AddMenuEntry(
@@ -575,7 +629,20 @@ void UToolMenus::PopulateMenuBuilder(FMenuBuilder& MenuBuilder, UToolMenu* MenuD
 			continue;
 		}
 
-		MenuBuilder.BeginSection(Section.Name, Section.Label);
+		if (bIsEditing)
+		{
+			// Always provide label when editing so we have area to drag/drop and hide sections
+			FText LabelText = Section.Label.Get();
+			if (LabelText.IsEmpty())
+			{
+				LabelText = FText::FromName(Section.Name);
+			}
+			MenuBuilder.BeginSection(Section.Name, LabelText);
+		}
+		else
+		{
+			MenuBuilder.BeginSection(Section.Name, Section.Label);
+		}
 
 		for (FToolMenuEntry& Block : Section.Blocks)
 		{
@@ -1088,7 +1155,7 @@ void UToolMenus::AddReferencedObjects(UObject* InThis, FReferenceCollector& Coll
 	Super::AddReferencedObjects(InThis, Collector);
 }
 
-UToolMenu* UToolMenus::GenerateMenu(const FName Name, FToolMenuContext& InMenuContext)
+UToolMenu* UToolMenus::GenerateMenu(const FName Name, const FToolMenuContext& InMenuContext)
 {
 	return GenerateMenu(CollectHierarchy(Name), InMenuContext);
 }
@@ -1106,13 +1173,13 @@ UToolMenu* UToolMenus::GenerateMenu(const TArray<UToolMenu*>& Hierarchy, const F
 	return GeneratedMenu;
 }
 
-TSharedRef< class SWidget > UToolMenus::GenerateWidget(const FName InName, FToolMenuContext& InMenuContext)
+TSharedRef< class SWidget > UToolMenus::GenerateWidget(const FName InName, const FToolMenuContext& InMenuContext)
 {
 	UToolMenu* Generated = GenerateMenu(InName, InMenuContext);
 	return GenerateWidget(Generated);
 }
 
-TSharedRef<SWidget> UToolMenus::GenerateWidget(const TArray<UToolMenu*>& Hierarchy, FToolMenuContext& InMenuContext)
+TSharedRef<SWidget> UToolMenus::GenerateWidget(const TArray<UToolMenu*>& Hierarchy, const FToolMenuContext& InMenuContext)
 {
 	if (Hierarchy.Num() == 0)
 	{
@@ -1136,9 +1203,29 @@ TSharedRef<SWidget> UToolMenus::GenerateWidget(UToolMenu* GeneratedMenu)
 	GeneratedMenuWidget.GeneratedMenu->Context = GeneratedMenu->Context;
 	GeneratedMenuWidget.GeneratedMenu->StyleSet = GeneratedMenu->StyleSet;
 
-	if (GeneratedMenu->MenuType == EMultiBoxType::Menu)
+	if (GeneratedMenu->IsEditing())
 	{
-		FMenuBuilder MenuBuilder(GeneratedMenu->bShouldCloseWindowAfterMenuSelection, GeneratedMenu->Context.CommandList, GeneratedMenu->Context.GetAllExtenders(), GeneratedMenu->bCloseSelfOnly, GeneratedMenu->StyleSet, GeneratedMenu->bSearchable);
+		// Convert toolbar into menu during editing
+		if (GeneratedMenu->MenuType == EMultiBoxType::ToolBar || GeneratedMenu->MenuType == EMultiBoxType::VerticalToolBar)
+		{
+			for (FToolMenuSection& Section : GeneratedMenu->Sections)
+			{
+				for (FToolMenuEntry& Entry : Section.Blocks)
+				{
+					ModifyEntryForEditDialog(Entry);
+				}
+			}
+		}
+
+		FMenuBuilder MenuBuilder(GeneratedMenu->bShouldCloseWindowAfterMenuSelection, GeneratedMenu->Context.CommandList, GeneratedMenu->Context.GetAllExtenders(), GeneratedMenu->bCloseSelfOnly, GeneratedMenu->StyleSet, GeneratedMenu->bSearchable, GeneratedMenu->MenuName);
+		PopulateMenuBuilder(MenuBuilder, GeneratedMenu);
+		TSharedRef<SWidget> Result = MenuBuilder.MakeWidget();
+		GeneratedMenuWidget.Widget = Result;
+		return Result;
+	}
+	else if (GeneratedMenu->MenuType == EMultiBoxType::Menu)
+	{
+		FMenuBuilder MenuBuilder(GeneratedMenu->bShouldCloseWindowAfterMenuSelection, GeneratedMenu->Context.CommandList, GeneratedMenu->Context.GetAllExtenders(), GeneratedMenu->bCloseSelfOnly, GeneratedMenu->StyleSet, GeneratedMenu->bSearchable, GeneratedMenu->MenuName);
 		PopulateMenuBuilder(MenuBuilder, GeneratedMenu);
 		TSharedRef<SWidget> Result = MenuBuilder.MakeWidget();
 		GeneratedMenuWidget.Widget = Result;
@@ -1146,7 +1233,7 @@ TSharedRef<SWidget> UToolMenus::GenerateWidget(UToolMenu* GeneratedMenu)
 	}
 	else if (GeneratedMenu->MenuType == EMultiBoxType::MenuBar)
 	{
-		FMenuBarBuilder MenuBarBuilder(GeneratedMenu->Context.CommandList, GeneratedMenu->Context.GetAllExtenders(), GeneratedMenu->StyleSet);
+		FMenuBarBuilder MenuBarBuilder(GeneratedMenu->Context.CommandList, GeneratedMenu->Context.GetAllExtenders(), GeneratedMenu->StyleSet, GeneratedMenu->MenuName);
 		PopulateMenuBarBuilder(MenuBarBuilder, GeneratedMenu);
 		TSharedRef<SWidget> Result = MenuBarBuilder.MakeWidget();
 		GeneratedMenuWidget.Widget = Result;
@@ -1300,7 +1387,7 @@ bool UToolMenus::RefreshMenuWidget(const FName InName, FGeneratedToolMenuWidget&
 	return true;
 }
 
-UToolMenu* UToolMenus::GenerateMenuAsBuilder(const UToolMenu* InMenu, FToolMenuContext& InMenuContext)
+UToolMenu* UToolMenus::GenerateMenuAsBuilder(const UToolMenu* InMenu, const FToolMenuContext& InMenuContext)
 {
 	TArray<UToolMenu*> Hierarchy = CollectHierarchy(InMenu->MenuName);
 
diff --git a/Engine/Source/Developer/ToolMenus/Public/ToolMenu.h b/Engine/Source/Developer/ToolMenus/Public/ToolMenu.h
--- a/Engine/Source/Developer/ToolMenus/Public/ToolMenu.h
+++ b/Engine/Source/Developer/ToolMenus/Public/ToolMenu.h
@@ -9,11 +9,12 @@
 #include "ToolMenuSection.h"
 #include "ToolMenuContext.h"
 #include "UObject/Object.h"
+#include "Framework/MultiBox/ToolMenuBase.h"
 
 #include "ToolMenu.generated.h"
 
 UCLASS(BlueprintType)
-class TOOLMENUS_API UToolMenu : public UObject
+class TOOLMENUS_API UToolMenu : public UToolMenuBase
 {
 	GENERATED_BODY()
 
@@ -55,6 +56,13 @@ public:
 
 	bool IsRegistered() const { return bRegistered; }
 
+	//~ Begin UToolMenuBase Interface
+	virtual bool IsEditing() const override;
+	virtual FName GetSectionName(const FName InEntryName) const override;
+	virtual FCustomizedToolMenu* FindMenuCustomization() const override;
+	virtual FCustomizedToolMenu* AddMenuCustomization() const override;
+	//~ End UToolMenuBase Interface
+
 	template <typename TContextType>
 	TContextType* FindContext() const
 	{
diff --git a/Engine/Source/Developer/ToolMenus/Public/ToolMenuContext.h b/Engine/Source/Developer/ToolMenus/Public/ToolMenuContext.h
--- a/Engine/Source/Developer/ToolMenus/Public/ToolMenuContext.h
+++ b/Engine/Source/Developer/ToolMenus/Public/ToolMenuContext.h
@@ -69,6 +69,10 @@ public:
 	friend class UToolMenus;
 	friend struct FToolMenuEntry;
 
+	bool IsEditing() const { return bIsEditing; }
+
+	bool bIsEditing;
+
 private:
 
 	UPROPERTY()
diff --git a/Engine/Source/Developer/ToolMenus/Public/ToolMenuEntry.h b/Engine/Source/Developer/ToolMenus/Public/ToolMenuEntry.h
--- a/Engine/Source/Developer/ToolMenus/Public/ToolMenuEntry.h
+++ b/Engine/Source/Developer/ToolMenus/Public/ToolMenuEntry.h
@@ -104,6 +104,8 @@ struct TOOLMENUS_API FToolMenuEntry
 
 	bool IsSubMenu() const { return SubMenuData.bIsSubMenu; }
 
+	bool IsConstructLegacy() const { return ConstructLegacy.IsBound(); }
+
 	const FUIAction* GetActionForCommand(const FToolMenuContext& InContext, TSharedPtr<const FUICommandList>& OutCommandList) const;
 
 	void SetCommandList(const TSharedPtr<const FUICommandList>& InCommandList);
diff --git a/Engine/Source/Developer/ToolMenus/Public/ToolMenuSection.h b/Engine/Source/Developer/ToolMenus/Public/ToolMenuSection.h
--- a/Engine/Source/Developer/ToolMenus/Public/ToolMenuSection.h
+++ b/Engine/Source/Developer/ToolMenus/Public/ToolMenuSection.h
@@ -44,13 +44,16 @@ public:
 
 	FToolMenuEntry& AddSubMenu(const FName InName, const TAttribute<FText>& InLabel, const TAttribute<FText>& InToolTip, const FNewToolMenuChoice& InMakeMenu, const FToolUIActionChoice& InAction, const EUserInterfaceActionType InUserInterfaceActionType, bool bInOpenSubMenuOnClick = false, const TAttribute<FSlateIcon>& InIcon = TAttribute<FSlateIcon>(), const bool bShouldCloseWindowAfterMenuSelection = true);
 	FToolMenuEntry& AddSubMenu(const FName InName, const TAttribute<FText>& InLabel, const TAttribute<FText>& InToolTip, const FNewToolMenuChoice& InMakeMenu, bool bInOpenSubMenuOnClick = false, const TAttribute<FSlateIcon>& InIcon = TAttribute<FSlateIcon>(), const bool bShouldCloseWindowAfterMenuSelection = true);
-	
+	FToolMenuEntry& AddSubMenu(const FName InName, const FToolUIActionChoice& InAction, const TSharedRef<SWidget>& InWidget, const FNewToolMenuChoice& InMakeMenu, bool bShouldCloseWindowAfterMenuSelection = true);
+
 	template <typename TContextType>
 	TContextType* FindContext() const
 	{
 		return Context.Find<TContextType>();
 	}
 
+	FToolMenuEntry* FindEntry(const FName InName);
+
 private:
 
 	void InitGeneratedSectionCopy(const FToolMenuSection& Source, FToolMenuContext& InContext);
diff --git a/Engine/Source/Developer/ToolMenus/Public/ToolMenus.h b/Engine/Source/Developer/ToolMenus/Public/ToolMenus.h
--- a/Engine/Source/Developer/ToolMenus/Public/ToolMenus.h
+++ b/Engine/Source/Developer/ToolMenus/Public/ToolMenus.h
@@ -32,36 +32,6 @@ class UToolMenuEntryScript;
 struct FToolMenuEntry;
 struct FToolMenuSection;
 
-USTRUCT()
-struct FCustomizedToolMenuSection
-{
-	GENERATED_BODY()
-
-	UPROPERTY()
-	FName Name;
-
-	UPROPERTY()
-	TArray<FName> Items;
-};
-
-USTRUCT()
-struct FCustomizedToolMenu
-{
-	GENERATED_BODY()
-
-	UPROPERTY()
-	FName Name;
-
-	UPROPERTY()
-	TArray<FCustomizedToolMenuSection> Sections;
-
-	UPROPERTY()
-	TArray<FName> HiddenItems;
-
-	UPROPERTY()
-	TArray<FName> HiddenSections;
-};
-
 USTRUCT()
 struct FGeneratedToolMenuWidget
 {
@@ -167,7 +137,7 @@ public:
 	 * @param	Context	Additional information specific to the menu being generated
 	 * @return	Widget to display
 	 */
-	TSharedRef<SWidget> GenerateWidget(const FName Name, FToolMenuContext& InMenuContext);
+	TSharedRef<SWidget> GenerateWidget(const FName Name, const FToolMenuContext& InMenuContext);
 
 
 
@@ -240,7 +210,7 @@ public:
 	 * @param	Context	Additional information specific to the menu being generated
 	 * @return	Widget to display
 	 */
-	TSharedRef<SWidget> GenerateWidget(const TArray<UToolMenu*>& Hierarchy, FToolMenuContext& InMenuContext);
+	TSharedRef<SWidget> GenerateWidget(const TArray<UToolMenu*>& Hierarchy, const FToolMenuContext& InMenuContext);
 
 	/**
 	 * Generate widget from a final collapsed menu. For advanced specialized use cases.
@@ -250,13 +220,13 @@ public:
 	TSharedRef<SWidget> GenerateWidget(UToolMenu* GeneratedMenu);
 	
 	/** Create a finalized menu that combines all parents used to generate a widget. Advanced special use cases only. */
-	UToolMenu* GenerateMenu(const FName Name, FToolMenuContext& InMenuContext);
+	UToolMenu* GenerateMenu(const FName Name, const FToolMenuContext& InMenuContext);
 
 	/** Create a finalized menu that combines given hierarchy array that will generate a widget. Advanced special use cases only. */
 	UToolMenu* GenerateMenu(const TArray<UToolMenu*>& Hierarchy, const FToolMenuContext& InMenuContext);
 
 	/** Create a finalized menu based on a custom crafted menu. Advanced special use cases only. */
-	UToolMenu* GenerateMenuAsBuilder(const UToolMenu* InMenu, FToolMenuContext& InMenuContext);
+	UToolMenu* GenerateMenuAsBuilder(const UToolMenu* InMenu, const FToolMenuContext& InMenuContext);
 
 	/** Bake final menu including calls to construction delegates, sorting, and customization */
 	void AssembleMenuHierarchy(UToolMenu* GeneratedMenu, const TArray<UToolMenu*>& Hierarchy);
@@ -279,6 +249,18 @@ public:
 	/** Timer function used to consolidate multiple duplicate requests into a single frame. */
 	void HandleNextTick();
 
+	/** Remove customization for a menu */
+	void RemoveCustomization(const FName InName);
+
+	/** Find customization settings for a menu */
+	FCustomizedToolMenu* FindMenuCustomization(const FName InName);
+
+	/** Find or add customization settings for a menu */
+	FCustomizedToolMenu* AddMenuCustomization(const FName InName);
+
+	/** Find index of customization settings for a menu */
+	int32 FindMenuCustomizationIndex(const FName InName);
+
 	/** Displaying extension points is for debugging menus */
 	DECLARE_DELEGATE_RetVal(bool, FShouldDisplayExtensionPoints);
 	FShouldDisplayExtensionPoints ShouldDisplayExtensionPoints;
@@ -292,6 +274,8 @@ public:
 	static void AddReferencedObjects(UObject* InThis, FReferenceCollector& Collector);
 	//~ End UObject Interface
 
+	FSlateIcon EditMenuIcon;
+
 private:
 
 	/** Sets a timer to be called next engine tick so that multiple repeated actions can be combined together. */
@@ -346,9 +330,6 @@ private:
 
 	void ApplyCustomization(UToolMenu* GeneratedMenu);
 
-	FCustomizedToolMenu* FindCustomizedMenu(const FName InName);
-	int32 FindCustomizedMenuIndex(const FName InName);
-
 	void UnregisterOwnerInternal(FToolMenuOwner Owner);
 
 	bool GetDisplayUIExtensionPoints() const;
diff --git a/Engine/Source/Editor/ContentBrowser/Private/AssetContextMenu.cpp b/Engine/Source/Editor/ContentBrowser/Private/AssetContextMenu.cpp
--- a/Engine/Source/Editor/ContentBrowser/Private/AssetContextMenu.cpp
+++ b/Engine/Source/Editor/ContentBrowser/Private/AssetContextMenu.cpp
@@ -145,7 +145,6 @@ TSharedRef<SWidget> FAssetContextMenu::MakeContextMenu(const TArray<FAssetData>&
 	TArray<UObject*> SelectedObjects;
 
 	// Create menu hierarchy based on class hierarchy
-	TArray<UToolMenu*> MenuHierarchy;
 	FName MenuName = BaseMenuName;
 	{
 		// Objects must be loaded for this operation... for now
@@ -171,79 +170,54 @@ TSharedRef<SWidget> FAssetContextMenu::MakeContextMenu(const TArray<FAssetData>&
 			}
 			ContextObject->CommonClass = CommonClass;
 
+			MenuName = UToolMenus::JoinMenuPaths(BaseMenuName, CommonClass->GetFName());
+
+			RegisterMenuHierarchy(CommonClass);
+
 			// Find asset actions for common class
 			FAssetToolsModule& AssetToolsModule = FModuleManager::LoadModuleChecked<FAssetToolsModule>("AssetTools");
 			TSharedPtr<IAssetTypeActions> CommonAssetTypeActions = AssetToolsModule.Get().GetAssetTypeActionsForClass(ContextObject->CommonClass).Pin();
-			UClass* AssetTypeActionsClass = nullptr;
 			if (CommonAssetTypeActions.IsValid() && CommonAssetTypeActions->HasActions(SelectedObjects))
 			{
 				ContextObject->CommonAssetTypeActions = CommonAssetTypeActions;
-				AssetTypeActionsClass = CommonAssetTypeActions->GetSupportedClass();
-				MenuName = UToolMenus::JoinMenuPaths(BaseMenuName, AssetTypeActionsClass->GetFName());
-			}
-
-			// Find menus for common class
-			UClass* ClassWithMenu = CommonClass;
-			while (ClassWithMenu)
-			{
-				FName CheckName = UToolMenus::JoinMenuPaths(BaseMenuName, ClassWithMenu->GetFName());
-				UToolMenu* FoundMenu = ToolMenus->FindMenu(CheckName);
-				if (FoundMenu && FoundMenu->Sections.Num() > 0)
-				{
-					// Change generated menu's name if necessary
-					if (MenuHierarchy.Num() == 0)
-					{
-						if (!AssetTypeActionsClass || ClassWithMenu->IsChildOf(AssetTypeActionsClass))
-						{
-							MenuName = CheckName;
-						}
-					}
-
-					MenuHierarchy.Add(FoundMenu);
-				}
-
-				ClassWithMenu = ClassWithMenu->GetSuperClass();
 			}
 		}
-
-		MenuHierarchy.Add(ToolMenus->FindMenu(BaseMenuName));
-
-		Algo::Reverse(MenuHierarchy);
 	}
 
 	FToolMenuContext MenuContext(InCommandList, MenuExtender, ContextObject);
-	UToolMenu* Menu = ToolMenus->GenerateMenu(MenuHierarchy, MenuContext);
-	Menu->MenuName = MenuName;
-
-	// Only add something if at least one asset is selected
-	if ( SelectedAssets.Num() )
-	{
-		// Add any type-specific context menu options
-		AddAssetTypeMenuOptions(Menu, SelectedObjects);
-
-		// Add imported asset context menu options
-		AddImportedAssetMenuOptions(Menu);
-
-		// Add quick access to common commands.
-		AddCommonMenuOptions(Menu);
+	return ToolMenus->GenerateWidget(MenuName, MenuContext);
+}
 
-		// Add quick access to view commands
-		AddExploreMenuOptions(Menu);
+void FAssetContextMenu::RegisterMenuHierarchy(UClass* InClass)
+{
+	static const FName BaseMenuName("ContentBrowser.AssetContextMenu");
 
-		// Add reference options
-		AddReferenceMenuOptions(Menu);
+	UToolMenus* ToolMenus = UToolMenus::Get();
 
-		// Add collection options
-		AddCollectionMenuOptions(Menu);
+	for (UClass* CurrentClass = InClass; CurrentClass; CurrentClass = CurrentClass->GetSuperClass())
+	{
+		FName CurrentMenuName = UToolMenus::JoinMenuPaths(BaseMenuName, CurrentClass->GetFName());
+		if (!ToolMenus->IsMenuRegistered(CurrentMenuName))
+		{
+			FName ParentMenuName;
+			UClass* ParentClass = CurrentClass->GetSuperClass();
+			if (ParentClass == UObject::StaticClass() || ParentClass == nullptr)
+			{
+				ParentMenuName = BaseMenuName;
+			}
+			else
+			{
+				ParentMenuName = UToolMenus::JoinMenuPaths(BaseMenuName, ParentClass->GetFName());
+			}
 
-		// Add documentation options
-		AddDocumentationMenuOptions(Menu);
+			ToolMenus->RegisterMenu(CurrentMenuName, ParentMenuName);
 
-		// Add source control options
-		AddSourceControlMenuOptions(Menu);
+			if (ParentMenuName == BaseMenuName)
+			{
+				break;
+			}
+		}
 	}
-	
-	return ToolMenus->GenerateWidget(Menu);
 }
 
 void FAssetContextMenu::RegisterContextMenu(const FName MenuName)
@@ -271,9 +245,46 @@ void FAssetContextMenu::RegisterContextMenu(const FName MenuName)
 				Context->CommonAssetTypeActions.Pin()->GetActions(Context->GetSelectedObjects(), MenuBuilder);
 			}
 		}));
+
+		Menu->AddDynamicSection("AddMenuOptions", FNewToolMenuDelegate::CreateSP(this, &FAssetContextMenu::AddMenuOptions));
 	}
 }
 
+void FAssetContextMenu::AddMenuOptions(UToolMenu* InMenu)
+{
+	UContentBrowserAssetContextMenuContext* Context = InMenu->FindContext<UContentBrowserAssetContextMenuContext>();
+	if (!Context || Context->SelectedObjects.Num() == 0)
+	{
+		return;
+	}
+
+	// Add any type-specific context menu options
+	AddAssetTypeMenuOptions(InMenu, Context->SelectedObjects.Num() > 0);
+
+	// Add imported asset context menu options
+	AddImportedAssetMenuOptions(InMenu);
+
+	// Add quick access to common commands.
+	AddCommonMenuOptions(InMenu);
+
+	// Add quick access to view commands
+	AddExploreMenuOptions(InMenu);
+
+	// Add reference options
+	AddReferenceMenuOptions(InMenu);
+
+	// Add collection options
+	AddCollectionMenuOptions(InMenu);
+
+	// Add documentation options
+	AddDocumentationMenuOptions(InMenu);
+
+	// Add source control options
+	AddSourceControlMenuOptions(InMenu);
+
+	UE_LOG(LogTemp, Log, TEXT("AddMenuOptions"));
+}
+
 void FAssetContextMenu::SetSelectedAssets(const TArray<FAssetData>& InSelectedAssets)
 {
 	SelectedAssets = InSelectedAssets;
@@ -1342,11 +1353,11 @@ bool FAssetContextMenu::AddDocumentationMenuOptions(UToolMenu* Menu)
 	return bAddedOption;
 }
 
-bool FAssetContextMenu::AddAssetTypeMenuOptions(UToolMenu* Menu, TArray<UObject*>& SelectedObjects)
+bool FAssetContextMenu::AddAssetTypeMenuOptions(UToolMenu* Menu, bool bHasObjectsSelected)
 {
 	bool bAnyTypeOptions = false;
 
-	if (SelectedObjects.Num() > 0)
+	if (bHasObjectsSelected)
 	{
 		// Label "GetAssetActions" section
 		UContentBrowserAssetContextMenuContext* Context = Menu->FindContext<UContentBrowserAssetContextMenuContext>();
diff --git a/Engine/Source/Editor/ContentBrowser/Private/AssetContextMenu.h b/Engine/Source/Editor/ContentBrowser/Private/AssetContextMenu.h
--- a/Engine/Source/Editor/ContentBrowser/Private/AssetContextMenu.h
+++ b/Engine/Source/Editor/ContentBrowser/Private/AssetContextMenu.h
@@ -117,8 +117,15 @@ private:
 	void ExecuteDuplicate();
 
 private:
+
+	/** Registers all unregistered menus in the hierarchy for a class */
+	static void RegisterMenuHierarchy(UClass* InClass);
+
+	/** Adds options to menu */
+	void AddMenuOptions(UToolMenu* InMenu);
+
 	/** Adds asset type-specific menu options to a menu builder. Returns true if any options were added. */
-	bool AddAssetTypeMenuOptions(UToolMenu* Menu, TArray<UObject*>& SelectedObjects);
+	bool AddAssetTypeMenuOptions(UToolMenu* Menu, bool bHasObjectsSelected);
 
 	/** Adds asset type-specific menu options to a menu builder. Returns true if any options were added. */
 	bool AddImportedAssetMenuOptions(UToolMenu* Menu);
diff --git a/Engine/Source/Runtime/Slate/Private/Framework/MultiBox/MultiBoxBuilder.cpp b/Engine/Source/Runtime/Slate/Private/Framework/MultiBox/MultiBoxBuilder.cpp
--- a/Engine/Source/Runtime/Slate/Private/Framework/MultiBox/MultiBoxBuilder.cpp
+++ b/Engine/Source/Runtime/Slate/Private/Framework/MultiBox/MultiBoxBuilder.cpp
@@ -16,10 +16,11 @@
 #include "Framework/MultiBox/SWidgetBlock.h"
 #include "Framework/MultiBox/SGroupMarkerBlock.h"
 
-FMultiBoxBuilder::FMultiBoxBuilder( const EMultiBoxType InType, FMultiBoxCustomization InCustomization, const bool bInShouldCloseWindowAfterMenuSelection, const TSharedPtr< const FUICommandList >& InCommandList, TSharedPtr<FExtender> InExtender, FName InTutorialHighlightName )
-	: MultiBox( FMultiBox::Create( InType, InCustomization, bInShouldCloseWindowAfterMenuSelection ) )
+FMultiBoxBuilder::FMultiBoxBuilder( const EMultiBoxType InType, FMultiBoxCustomization InCustomization, const bool bInShouldCloseWindowAfterMenuSelection, const TSharedPtr< const FUICommandList >& InCommandList, TSharedPtr<FExtender> InExtender, FName InTutorialHighlightName, FName InMenuName )
+	: MultiBox( FMultiBox::Create( InType, InMenuName != NAME_None ? FMultiBoxCustomization::AllowCustomization(InMenuName) : InCustomization, bInShouldCloseWindowAfterMenuSelection ) )
 	, CommandListStack()
 	, TutorialHighlightName(InTutorialHighlightName)
+	, MenuName(InMenuName)
 {
 	CommandListStack.Push( InCommandList );
 	ExtenderStack.Push(InExtender);
@@ -118,8 +119,8 @@ static FName GenerateTutorialIdentfierName(FName InContainerName, FName InElemen
 	}
 }
 
-FBaseMenuBuilder::FBaseMenuBuilder( const EMultiBoxType InType, const bool bInShouldCloseWindowAfterMenuSelection, TSharedPtr< const FUICommandList > InCommandList, bool bInCloseSelfOnly, TSharedPtr<FExtender> InExtender, const ISlateStyle* InStyleSet, FName InTutorialHighlightName )
-	: FMultiBoxBuilder( InType, FMultiBoxCustomization::None, bInShouldCloseWindowAfterMenuSelection, InCommandList, InExtender, InTutorialHighlightName )
+FBaseMenuBuilder::FBaseMenuBuilder( const EMultiBoxType InType, const bool bInShouldCloseWindowAfterMenuSelection, TSharedPtr< const FUICommandList > InCommandList, bool bInCloseSelfOnly, TSharedPtr<FExtender> InExtender, const ISlateStyle* InStyleSet, FName InTutorialHighlightName, FName InMenuName )
+	: FMultiBoxBuilder( InType, FMultiBoxCustomization::None, bInShouldCloseWindowAfterMenuSelection, InCommandList, InExtender, InTutorialHighlightName, InMenuName )
 	, bCloseSelfOnly( bInCloseSelfOnly )
 {
 	MultiBox->SetStyle(InStyleSet, "Menu");
diff --git a/Engine/Source/Runtime/Slate/Private/Framework/MultiBox/MultiBoxCustomization.cpp b/Engine/Source/Runtime/Slate/Private/Framework/MultiBox/MultiBoxCustomization.cpp
--- a/Engine/Source/Runtime/Slate/Private/Framework/MultiBox/MultiBoxCustomization.cpp
+++ b/Engine/Source/Runtime/Slate/Private/Framework/MultiBox/MultiBoxCustomization.cpp
@@ -1,6 +1,7 @@
 // Copyright 1998-2019 Epic Games, Inc. All Rights Reserved.
 
 #include "Framework/MultiBox/MultiBoxCustomization.h"
+#include "Framework/MultiBox/ToolMenuBase.h"
 #include "Misc/ConfigCacheIni.h"
 #include "Serialization/JsonReader.h"
 #include "Serialization/JsonSerializer.h"
@@ -8,6 +9,89 @@
 #include "Widgets/Layout/SBorder.h"
 #include "Misc/RemoteConfigIni.h"
 #include "Framework/Commands/UICommandDragDropOp.h"
+#include "Widgets/Colors/SColorBlock.h"
+
+FCustomizedToolMenuSection* FCustomizedToolMenu::FindSection(const FName InSectionName)
+{
+	for (FCustomizedToolMenuSection& Section : Sections)
+	{
+		if (Section.Name == InSectionName)
+		{
+			return &Section;
+		}
+	}
+
+	return nullptr;
+}
+
+const FCustomizedToolMenuSection* FCustomizedToolMenu::FindSection(const FName InSectionName) const
+{
+	for (const FCustomizedToolMenuSection& Section : Sections)
+	{
+		if (Section.Name == InSectionName)
+		{
+			return &Section;
+		}
+	}
+
+	return nullptr;
+}
+
+FCustomizedToolMenuEntry* FCustomizedToolMenu::FindEntry(const FName InEntryName, FName* OutSectionName)
+{
+	for (FCustomizedToolMenuSection& Section : Sections)
+	{
+		for (FCustomizedToolMenuEntry& Entry : Section.Entries)
+		{
+			if (Entry.Name == InEntryName)
+			{
+				if (OutSectionName)
+				{
+					*OutSectionName = Section.Name;
+				}
+				return &Entry;
+			}
+		}
+	}
+
+	return nullptr;
+}
+
+const FCustomizedToolMenuEntry* FCustomizedToolMenu::FindEntry(const FName InEntryName, FName* OutSectionName) const
+{
+	for (const FCustomizedToolMenuSection& Section : Sections)
+	{
+		for (const FCustomizedToolMenuEntry& Entry : Section.Entries)
+		{
+			if (Entry.Name == InEntryName)
+			{
+				if (OutSectionName)
+				{
+					*OutSectionName = Section.Name;
+				}
+				return &Entry;
+			}
+		}
+	}
+
+	return nullptr;
+}
+
+bool FCustomizedToolMenu::IsSectionHidden(const FName InSectionName) const
+{
+	return HiddenSections.Contains(InSectionName);
+}
+
+bool FCustomizedToolMenu::IsEntryHidden(const FName InEntryName) const
+{
+	return HiddenEntries.Contains(InEntryName);
+}
+
+void FCustomizedToolMenu::UpdateFromMultiBox(const TSharedRef<const FMultiBox>& InMultiBox)
+{
+
+}
+
 
 void SCustomToolbarPreviewWidget::Construct( const FArguments& InArgs )
 {
diff --git a/Engine/Source/Runtime/Slate/Public/Framework/MultiBox/MultiBoxBuilder.h b/Engine/Source/Runtime/Slate/Public/Framework/MultiBox/MultiBoxBuilder.h
--- a/Engine/Source/Runtime/Slate/Public/Framework/MultiBox/MultiBoxBuilder.h
+++ b/Engine/Source/Runtime/Slate/Public/Framework/MultiBox/MultiBoxBuilder.h
@@ -42,7 +42,7 @@ public:
 	 * @param	InCommandList	The action list that maps command infos to delegates that should be called for each command associated with a multiblock widget.  This can be modified after the MultiBox is created by calling the PushCommandList() and PopCommandList() methods.
 	 * @param	InTutorialHighlightName	Optional name to identify this widget and highlight during tutorials
 	 */
-	FMultiBoxBuilder( const EMultiBoxType InType, FMultiBoxCustomization InCustomization, const bool bInShouldCloseWindowAfterMenuSelection, const TSharedPtr< const FUICommandList >& InCommandList, TSharedPtr<FExtender> InExtender = TSharedPtr<FExtender>(), FName InTutorialHighlightName = NAME_None );
+	FMultiBoxBuilder( const EMultiBoxType InType, FMultiBoxCustomization InCustomization, const bool bInShouldCloseWindowAfterMenuSelection, const TSharedPtr< const FUICommandList >& InCommandList, TSharedPtr<FExtender> InExtender = TSharedPtr<FExtender>(), FName InTutorialHighlightName = NAME_None, FName InMenuName = NAME_None );
 
 	virtual ~FMultiBoxBuilder() {}
 
@@ -140,6 +140,9 @@ protected:
 
 	/** Name to identify this widget and highlight during tutorials */
 	FName TutorialHighlightName;
+
+	/** Name of the menu */
+	FName MenuName;
 };
 
 
@@ -161,7 +164,7 @@ public:
 	 * @param	bInCloseSelfOnly	True if clicking on a menu entry closes itself only and its children but not the entire stack 
 	 * @param	InTutorialHighlightName	Optional name to identify this widget and highlight during tutorials
 	 */
-	FBaseMenuBuilder( const EMultiBoxType InType, const bool bInShouldCloseWindowAfterMenuSelection, TSharedPtr< const FUICommandList > InCommandList, bool bInCloseSelfOnly, TSharedPtr<FExtender> InExtender = TSharedPtr<FExtender>(), const ISlateStyle* InStyleSet = &FCoreStyle::Get(), FName InTutorialHighlightName = NAME_None );
+	FBaseMenuBuilder( const EMultiBoxType InType, const bool bInShouldCloseWindowAfterMenuSelection, TSharedPtr< const FUICommandList > InCommandList, bool bInCloseSelfOnly, TSharedPtr<FExtender> InExtender = TSharedPtr<FExtender>(), const ISlateStyle* InStyleSet = &FCoreStyle::Get(), FName InTutorialHighlightName = NAME_None, FName InMenuName = NAME_None );
 
 	/**
 	 * Adds a menu entry
@@ -210,10 +213,11 @@ public:
 	 * @param	InCommandList	The action list that maps command infos to delegates that should be called for each command associated with a multiblock widget
 	 * @param	bInCloseSelfOnly	True if clicking on a menu entry closes itself only and its children but not the entire stack
 	 */
-	FMenuBuilder( const bool bInShouldCloseWindowAfterMenuSelection, TSharedPtr< const FUICommandList > InCommandList, TSharedPtr<FExtender> InExtender = TSharedPtr<FExtender>(), const bool bCloseSelfOnly = false, const ISlateStyle* InStyleSet = &FCoreStyle::Get(), bool bInSearchable = true )
-		: FBaseMenuBuilder( EMultiBoxType::Menu, bInShouldCloseWindowAfterMenuSelection, InCommandList, bCloseSelfOnly, InExtender, InStyleSet )
+	FMenuBuilder( const bool bInShouldCloseWindowAfterMenuSelection, TSharedPtr< const FUICommandList > InCommandList, TSharedPtr<FExtender> InExtender = TSharedPtr<FExtender>(), const bool bCloseSelfOnly = false, const ISlateStyle* InStyleSet = &FCoreStyle::Get(), bool bInSearchable = true, FName InMenuName = NAME_None)
+		: FBaseMenuBuilder( EMultiBoxType::Menu, bInShouldCloseWindowAfterMenuSelection, InCommandList, bCloseSelfOnly, InExtender, InStyleSet, NAME_None, InMenuName )
 		, bSectionNeedsToBeApplied(false)
 		, bSearchable(bInSearchable)
+		, bIsEditing(false)
 	{
 		if(bSearchable)
 		{
@@ -284,6 +288,8 @@ public:
 	*/
 	void AddSearchWidget();
 
+	void SetIsEditing(bool bInIsEditing) { bIsEditing = bInIsEditing; }
+
 protected:
 	/** FMultiBoxBuilder interface */
 	virtual void ApplyHook(FName InExtensionHook, EExtensionHook::Position HookPosition) override;
@@ -324,6 +330,9 @@ private:
 
 	/** Whether this menu is searchable */
 	bool bSearchable;
+
+	/** Whether menu is currently being edited */
+	bool bIsEditing;
 };
 
 
@@ -341,8 +350,8 @@ public:
 	 *
 	 * @param	InCommandList	The action list that maps command infos to delegates that should be called for each command associated with a multiblock widget
 	 */
-	FMenuBarBuilder( TSharedPtr< const FUICommandList > InCommandList, TSharedPtr<FExtender> InExtender = TSharedPtr<FExtender>(), const ISlateStyle* InStyleSet = &FCoreStyle::Get())
-		: FBaseMenuBuilder( EMultiBoxType::MenuBar, false, InCommandList, false, InExtender, InStyleSet )
+	FMenuBarBuilder( TSharedPtr< const FUICommandList > InCommandList, TSharedPtr<FExtender> InExtender = TSharedPtr<FExtender>(), const ISlateStyle* InStyleSet = &FCoreStyle::Get(), FName InMenuName = NAME_None)
+		: FBaseMenuBuilder( EMultiBoxType::MenuBar, false, InCommandList, false, InExtender, InStyleSet, NAME_None, InMenuName )
 	{
 	}
 
diff --git a/Engine/Source/Runtime/Slate/Public/Framework/MultiBox/ToolMenuBase.h b/Engine/Source/Runtime/Slate/Public/Framework/MultiBox/ToolMenuBase.h
new file mode 100644
--- /dev/null
+++ b/Engine/Source/Runtime/Slate/Public/Framework/MultiBox/ToolMenuBase.h
@@ -0,0 +1,76 @@
+// Copyright 1998-2019 Epic Games, Inc. All Rights Reserved.
+
+#pragma once
+
+#include "CoreMinimal.h"
+
+#include "ToolMenuBase.generated.h"
+
+class FMultiBox;
+
+USTRUCT()
+struct FCustomizedToolMenuEntry
+{
+	GENERATED_BODY()
+
+	UPROPERTY()
+	FName Name;
+};
+
+USTRUCT()
+struct FCustomizedToolMenuSection
+{
+	GENERATED_BODY()
+		
+	UPROPERTY()
+	FName Name;
+
+	UPROPERTY()
+	TArray<FCustomizedToolMenuEntry> Entries;
+};
+
+USTRUCT()
+struct SLATE_API FCustomizedToolMenu
+{
+	GENERATED_BODY()
+
+	UPROPERTY()
+	FName Name;
+
+	UPROPERTY()
+	TArray<FCustomizedToolMenuSection> Sections;
+
+	UPROPERTY()
+	TArray<FName> HiddenSections;
+
+	UPROPERTY()
+	TArray<FName> HiddenEntries;
+
+	bool IsSectionHidden(const FName InSectionName) const;
+	bool IsEntryHidden(const FName InSectionName) const;
+
+	FCustomizedToolMenuSection* FindSection(const FName InSectionName);
+	const FCustomizedToolMenuSection* FindSection(const FName InSectionName) const;
+
+	FCustomizedToolMenuEntry* FindEntry(const FName InEntryName, FName* OutSectionName = nullptr);
+	const FCustomizedToolMenuEntry* FindEntry(const FName InEntryName, FName* OutSectionName = nullptr) const;
+
+	/* Updates re-positioning of sections and entries
+	 * Does not need to update hidden state as MultiBlocks/Widgets do not store hidden state of each entry and section themselves
+	 */
+	void UpdateFromMultiBox(const TSharedRef<const FMultiBox>& InMultiBox);
+};
+
+UCLASS(Abstract)
+class SLATE_API UToolMenuBase : public UObject
+{
+	GENERATED_BODY()
+
+public:
+
+	virtual bool IsEditing() const { return false; }
+	virtual FName GetSectionName(const FName InEntryName) const { return NAME_None; }
+
+	virtual FCustomizedToolMenu* FindMenuCustomization() const { return nullptr; }
+	virtual FCustomizedToolMenu* AddMenuCustomization() const { return nullptr; }
+};
