Tool menu's module load/unload functions are now called
#rnx

diff --git a/Engine/Source/Developer/ToolMenus/Private/ToolMenus.cpp b/Engine/Source/Developer/ToolMenus/Private/ToolMenus.cpp
--- a/Engine/Source/Developer/ToolMenus/Private/ToolMenus.cpp
+++ b/Engine/Source/Developer/ToolMenus/Private/ToolMenus.cpp
@@ -14,6 +14,9 @@
 
 #define LOCTEXT_NAMESPACE "ToolMenuSubsystem"
 
+UToolMenus* UToolMenus::Singleton = nullptr;
+bool UToolMenus::bHasShutDown = false;
+
 FAutoConsoleCommand ToolMenusRefreshMenuWidget = FAutoConsoleCommand(
 	TEXT("ToolMenus.RefreshAllWidgets"),
 	TEXT("Refresh All Tool Menu Widgets"),
@@ -64,26 +67,6 @@ FToolUIActionChoice::FToolUIActionChoice(const TSharedPtr< const FUICommandInfo
 	}
 }
 
-void ModifyEntryForEditDialog(FToolMenuEntry& Entry)
-{
-	if (Entry.Type == EMultiBlockType::ToolBarButton)
-	{
-		Entry.Type = EMultiBlockType::MenuEntry;
-	}
-	else if (Entry.Type == EMultiBlockType::ToolBarComboButton)
-	{
-		Entry.Type = EMultiBlockType::MenuEntry;
-		if (Entry.ToolBarData.bSimpleComboBox)
-		{
-			Entry.SubMenuData.bIsSubMenu = true;
-		}
-	}
-	else if (Entry.Type == EMultiBlockType::ToolBarSeparator)
-	{
-		Entry.Type = EMultiBlockType::MenuSeparator;
-	}
-}
-
 UToolMenus::UToolMenus() :
 	bNextTickTimerIsSet(false),
 	bRefreshWidgetsNextTick(false),
@@ -93,16 +76,30 @@ UToolMenus::UToolMenus() :
 
 UToolMenus* UToolMenus::Get()
 {
-	static UToolMenus* Singleton = nullptr;
-	if (!Singleton)
+	if (!Singleton && !bHasShutDown)
 	{
+		// Required for StartupModule and ShutdownModule to be called and FindModule to list the ToolsMenus module
+		FModuleManager::LoadModuleChecked<IToolMenusModule>("ToolMenus");
+
 		Singleton = NewObject<UToolMenus>();
 		Singleton->AddToRoot();
 		check(Singleton);
 	}
+
 	return Singleton;
 }
 
+void UToolMenus::BeginDestroy()
+{
+	if (Singleton == this)
+	{
+		bHasShutDown = true;
+		Singleton = nullptr;
+	}
+
+	Super::BeginDestroy();
+}
+
 bool UToolMenus::IsToolMenuUIEnabled()
 {
 	if (!FSlateApplication::IsInitialized())
@@ -1261,6 +1258,26 @@ TSharedRef<SWidget> UToolMenus::GenerateWidget(UToolMenu* GeneratedMenu)
 	return SNullWidget::NullWidget;
 }
 
+void UToolMenus::ModifyEntryForEditDialog(FToolMenuEntry& Entry)
+{
+	if (Entry.Type == EMultiBlockType::ToolBarButton)
+	{
+		Entry.Type = EMultiBlockType::MenuEntry;
+	}
+	else if (Entry.Type == EMultiBlockType::ToolBarComboButton)
+	{
+		Entry.Type = EMultiBlockType::MenuEntry;
+		if (Entry.ToolBarData.bSimpleComboBox)
+		{
+			Entry.SubMenuData.bIsSubMenu = true;
+		}
+	}
+	else if (Entry.Type == EMultiBlockType::ToolBarSeparator)
+	{
+		Entry.Type = EMultiBlockType::MenuSeparator;
+	}
+}
+
 void UToolMenus::AssignSetTimerForNextTickDelegate(const FSimpleDelegate& InDelegate)
 {
 	SetTimerForNextTickDelegate = InDelegate;
diff --git a/Engine/Source/Developer/ToolMenus/Public/ToolMenus.h b/Engine/Source/Developer/ToolMenus/Public/ToolMenus.h
--- a/Engine/Source/Developer/ToolMenus/Public/ToolMenus.h
+++ b/Engine/Source/Developer/ToolMenus/Public/ToolMenus.h
@@ -272,6 +272,8 @@ public:
 
 	//~ Begin UObject Interface
 	static void AddReferencedObjects(UObject* InThis, FReferenceCollector& Collector);
+	virtual void BeginDestroy() override;
+	virtual bool IsDestructionThreadSafe() const override { return false; }
 	//~ End UObject Interface
 
 	FSlateIcon EditMenuIcon;
@@ -334,6 +336,8 @@ private:
 
 	bool GetDisplayUIExtensionPoints() const;
 
+	static void ModifyEntryForEditDialog(FToolMenuEntry& Entry);
+
 private:
 
 	UPROPERTY(EditAnywhere, Category = Misc)
@@ -356,6 +360,10 @@ private:
 	bool bNextTickTimerIsSet;
 	bool bRefreshWidgetsNextTick;
 	bool bCleanupStaleWidgetsNextTick;
+
+	
+	static UToolMenus* Singleton;
+	static bool bHasShutDown;
 };
 
 struct FToolMenuOwnerScoped
