Add option to prevent menu from being customized with menu editor dialog
#rnx

diff --git a/Engine/Source/Developer/ToolMenus/Public/ToolMenu.h b/Engine/Source/Developer/ToolMenus/Public/ToolMenu.h
--- a/Engine/Source/Developer/ToolMenus/Public/ToolMenu.h
+++ b/Engine/Source/Developer/ToolMenus/Public/ToolMenu.h
@@ -116,6 +116,10 @@ public:
 
 	UPROPERTY(EditAnywhere, BlueprintReadWrite, Category = "ToolBar")
 	bool bToolBarForceSmallIcons;
+		
+	/** Prevent menu from being customized */
+	UPROPERTY(EditAnywhere, BlueprintReadWrite, Category = "ToolBar")
+	bool bPreventCustomization;
 
 	UPROPERTY(EditAnywhere, BlueprintReadWrite, Category = "Tool Menus")
 	FToolMenuOwner MenuOwner;
