Fix garbage collection causing menus to dissapear

diff --git a/Engine/Source/Developer/ToolMenus/Private/ToolMenus.cpp b/Engine/Source/Developer/ToolMenus/Private/ToolMenus.cpp
--- a/Engine/Source/Developer/ToolMenus/Private/ToolMenus.cpp
+++ b/Engine/Source/Developer/ToolMenus/Private/ToolMenus.cpp
@@ -904,7 +904,7 @@ void UToolMenus::PopulateMenuBuilder(FMenuBuilder& MenuBuilder, UToolMenu* MenuD
 		MenuBuilder.EndSection();
 	}
 
-	AddReferencedContextObjects(MenuBuilder.GetMultiBox(), MenuData->Context);
+	AddReferencedContextObjects(MenuBuilder.GetMultiBox(), MenuData);
 }
 
 void UToolMenus::PopulateToolBarBuilder(FToolBarBuilder& ToolBarBuilder, UToolMenu* MenuData)
@@ -1022,7 +1022,7 @@ void UToolMenus::PopulateToolBarBuilder(FToolBarBuilder& ToolBarBuilder, UToolMe
 		ToolBarBuilder.EndSection();
 	}
 
-	AddReferencedContextObjects(ToolBarBuilder.GetMultiBox(), MenuData->Context);
+	AddReferencedContextObjects(ToolBarBuilder.GetMultiBox(), MenuData);
 }
 
 void UToolMenus::PopulateMenuBarBuilder(FMenuBarBuilder& MenuBarBuilder, UToolMenu* MenuData)
@@ -1052,7 +1052,7 @@ void UToolMenus::PopulateMenuBarBuilder(FMenuBarBuilder& MenuBarBuilder, UToolMe
 		}
 	}
 
-	AddReferencedContextObjects(MenuBarBuilder.GetMultiBox(), MenuData->Context);
+	AddReferencedContextObjects(MenuBarBuilder.GetMultiBox(), MenuData);
 }
 
 FOnGetContent UToolMenus::ConvertWidgetChoice(const FNewToolMenuWidgetChoice& Choice, const FToolMenuContext& Context) const
@@ -1283,19 +1283,18 @@ UObject* UToolMenus::FindContext(const FToolMenuContext& InContext, UClass* InCl
 	return InContext.FindByClass(InClass);
 }
 
-void UToolMenus::AddReferencedContextObjects(const TSharedRef<FMultiBox>& InMultiBox, const FToolMenuContext& InMenuContext)
+void UToolMenus::AddReferencedContextObjects(const TSharedRef<FMultiBox>& InMultiBox, const UToolMenu* InMenu)
 {
-	if (InMenuContext.ContextObjects.Num() == 0)
+	if (InMenu)
 	{
-		return;
-	}
-
-	TArray<UObject*>& References = WidgetObjectReferences.FindOrAdd(InMultiBox);
-	for (const TWeakObjectPtr<UObject>& WeakObject : InMenuContext.ContextObjects)
-	{
-		if (UObject* Object = WeakObject.Get())
+		TArray<const UObject*>& References = WidgetObjectReferences.FindOrAdd(InMultiBox);
+		References.AddUnique(InMenu);
+		for (const TWeakObjectPtr<UObject>& WeakObject : InMenu->Context.ContextObjects)
 		{
-			References.AddUnique(Object);
+			if (UObject* Object = WeakObject.Get())
+			{
+				References.AddUnique(Object);
+			}
 		}
 	}
 }
diff --git a/Engine/Source/Developer/ToolMenus/Public/ToolMenus.h b/Engine/Source/Developer/ToolMenus/Public/ToolMenus.h
--- a/Engine/Source/Developer/ToolMenus/Public/ToolMenus.h
+++ b/Engine/Source/Developer/ToolMenus/Public/ToolMenus.h
@@ -364,7 +364,7 @@ private:
 
 	void CopyMenuSettings(UToolMenu* GeneratedMenu, const UToolMenu* Other);
 
-	void AddReferencedContextObjects(const TSharedRef<FMultiBox>& InMultiBox, const FToolMenuContext& InMenuContext);
+	void AddReferencedContextObjects(const TSharedRef<FMultiBox>& InMultiBox, const UToolMenu* InMenu);
 
 	void ApplyCustomization(UToolMenu* GeneratedMenu);
 
@@ -385,7 +385,7 @@ private:
 	UPROPERTY()
 	TMap<FName, FGeneratedToolMenuWidgets> GeneratedMenuWidgets;
 
-	TMap<TWeakPtr<FMultiBox>, TArray<UObject*>> WidgetObjectReferences;
+	TMap<TWeakPtr<FMultiBox>, TArray<const UObject*>> WidgetObjectReferences;
 
 	TArray<FToolMenuOwner> OwnerStack;
 
