Tool menu startup callback fixed to handle being called before slate has been initialized
#rnx

diff --git a/Engine/Source/Developer/ToolMenus/Private/ToolMenus.cpp b/Engine/Source/Developer/ToolMenus/Private/ToolMenus.cpp
--- a/Engine/Source/Developer/ToolMenus/Private/ToolMenus.cpp
+++ b/Engine/Source/Developer/ToolMenus/Private/ToolMenus.cpp
@@ -16,6 +16,8 @@
 
 UToolMenus* UToolMenus::Singleton = nullptr;
 bool UToolMenus::bHasShutDown = false;
+FSimpleMulticastDelegate UToolMenus::StartupCallbacks;
+TOptional<FDelegateHandle> UToolMenus::InternalStartupCallbackHandle;
 
 FAutoConsoleCommand ToolMenusRefreshMenuWidget = FAutoConsoleCommand(
 	TEXT("ToolMenus.RefreshAllWidgets"),
@@ -94,6 +96,8 @@ void UToolMenus::BeginDestroy()
 {
 	if (Singleton == this)
 	{
+		UnregisterPrivateStartupCallback();
+
 		bHasShutDown = true;
 		Singleton = nullptr;
 	}
@@ -1805,6 +1809,63 @@ void UToolMenus::UnregisterStringCommandHandler(const FName InName)
 	StringCommandHandlers.Remove(InName);
 }
 
+FDelegateHandle UToolMenus::RegisterStartupCallback(const FSimpleMulticastDelegate::FDelegate& InDelegate)
+{
+	if (IsToolMenuUIEnabled() && UToolMenus::TryGet())
+	{
+		// Call immediately if systems are initialized
+		InDelegate.Execute();
+	}
+	else
+	{
+		// Defer call to occur after systems are initialized (slate and menus)
+		FDelegateHandle Result = StartupCallbacks.Add(InDelegate);
+
+		if (!InternalStartupCallbackHandle.IsSet())
+		{
+			InternalStartupCallbackHandle = FCoreDelegates::OnPostEngineInit.Add(FSimpleMulticastDelegate::FDelegate::CreateStatic(&UToolMenus::PrivateStartupCallback));
+		}
+
+		return Result;
+	}
+
+	return FDelegateHandle();
+}
+
+void UToolMenus::UnRegisterStartupCallback(const void* UserPointer)
+{
+	StartupCallbacks.RemoveAll(UserPointer);
+}
+
+void UToolMenus::UnRegisterStartupCallback(FDelegateHandle InHandle)
+{
+	StartupCallbacks.Remove(InHandle);
+}
+
+void UToolMenus::PrivateStartupCallback()
+{
+	UnregisterPrivateStartupCallback();
+
+	if (IsToolMenuUIEnabled() && UToolMenus::TryGet())
+	{
+		StartupCallbacks.Broadcast();
+		StartupCallbacks.Clear();
+	}
+}
+
+void UToolMenus::UnregisterPrivateStartupCallback()
+{
+	if (InternalStartupCallbackHandle.IsSet())
+	{
+		FDelegateHandle& Handle = InternalStartupCallbackHandle.GetValue();
+		if (Handle.IsValid())
+		{
+			FCoreDelegates::OnPostEngineInit.Remove(Handle);
+			Handle.Reset();
+		}
+	}
+}
+
 void UToolMenus::SaveCustomizations()
 {
 	SaveConfig();
diff --git a/Engine/Source/Developer/ToolMenus/Public/ToolMenus.h b/Engine/Source/Developer/ToolMenus/Public/ToolMenus.h
--- a/Engine/Source/Developer/ToolMenus/Public/ToolMenus.h
+++ b/Engine/Source/Developer/ToolMenus/Public/ToolMenus.h
@@ -83,33 +83,19 @@ public:
 	 */
 	static bool IsToolMenuUIEnabled();
 
-	static FDelegateHandle RegisterStartupCallback(const FSimpleMulticastDelegate::FDelegate& Delegate)
-	{
-		if (IsToolMenuUIEnabled())
-		{
-			if (UToolMenus::TryGet())
-			{
-				Delegate.Execute();
-			}
-			else
-			{
-				// Wait until UToolMenus has been initialized
-				return FCoreDelegates::OnPostEngineInit.Add(Delegate);
-			}
-		}
-
-		return FDelegateHandle();
-	}
+	/** 
+	 * Delays menu registration until safe and ready
+	 * Will not trigger if Slate does not end up being enabled after loading
+	 * Will not trigger when running commandlet, game, dedicated server or client only
+	 *
+	 */
+	static FDelegateHandle RegisterStartupCallback(const FSimpleMulticastDelegate::FDelegate& InDelegate);
 
-	static void UnRegisterStartupCallback(const void* UserPointer)
-	{
-		FCoreDelegates::OnPostEngineInit.RemoveAll(UserPointer);
-	}
+	/** Unregister a startup callback delegate by pointer */
+	static void UnRegisterStartupCallback(const void* UserPointer);
 
-	static void UnRegisterStartupCallback(FDelegateHandle InHandle)
-	{
-		FCoreDelegates::OnPostEngineInit.Remove(InHandle);
-	}
+	/** Unregister a startup callback delegate by handle */
+	static void UnRegisterStartupCallback(FDelegateHandle InHandle);
 
 	/**
 	 * Registers a menu by name
@@ -326,6 +312,9 @@ private:
 	/** Sets the current temporary menu owner. Used by FToolMenuEntryScoped */
 	void PopOwner(const FToolMenuOwner InOwner);
 
+	static void PrivateStartupCallback();
+	static void UnregisterPrivateStartupCallback();
+
 	UToolMenu* FindSubMenuToGenerateWith(const FName InParentName, const FName InChildName);
 
 	void FillMenuBarDropDown(FMenuBuilder& MenuBuilder, FName InParentName, FName InChildName, FToolMenuContext InMenuContext);
@@ -390,6 +379,8 @@ private:
 
 	static UToolMenus* Singleton;
 	static bool bHasShutDown;
+	static FSimpleMulticastDelegate StartupCallbacks;
+	static TOptional<FDelegateHandle> InternalStartupCallbackHandle;
 };
 
 struct FToolMenuOwnerScoped
diff --git a/Engine/Source/Editor/ToolMenusEditor/Private/ToolMenusEditorModule.cpp b/Engine/Source/Editor/ToolMenusEditor/Private/ToolMenusEditorModule.cpp
--- a/Engine/Source/Editor/ToolMenusEditor/Private/ToolMenusEditorModule.cpp
+++ b/Engine/Source/Editor/ToolMenusEditor/Private/ToolMenusEditorModule.cpp
@@ -18,6 +18,13 @@ FAutoConsoleVariableRef CVarRex(
 	TEXT("Edit menus in modal window"),
 	ECVF_Default);
 
+FAutoConsoleCommand ToolMenusEditMenusModeCVar = FAutoConsoleCommand(
+	TEXT("ToolMenus.EditMenusMode"),
+	TEXT("Enable edit menus mode toggle in level editor's windows menu"),
+	FConsoleCommandDelegate::CreateLambda([]() {
+		IToolMenusEditorModule::Get().RegisterShowEditMenusModeCheckbox();
+	}));
+
 /**
  * Implements the Tool menus module.
  */
