Add tool menu sub-menu support for custom type and action.
#rnx

diff --git a/Engine/Source/Developer/ToolMenus/Private/ToolMenus.cpp b/Engine/Source/Developer/ToolMenus/Private/ToolMenus.cpp
--- a/Engine/Source/Developer/ToolMenus/Private/ToolMenus.cpp
+++ b/Engine/Source/Developer/ToolMenus/Private/ToolMenus.cpp
@@ -836,10 +836,13 @@ void UToolMenus::PopulateMenuBuilder(FMenuBuilder& MenuBuilder, UToolMenu* MenuD
 							Block.Label,
 							Block.ToolTip,
 							NewMenuDelegate,
+							UIAction,
+							Block.Name,
+							Block.UserInterfaceActionType,
 							Block.SubMenuData.bOpenSubMenuOnClick,
 							Block.Icon.Get(),
-							Block.bShouldCloseWindowAfterMenuSelection,
-							Block.Name);
+							Block.bShouldCloseWindowAfterMenuSelection
+							);
 					}
 				}
 				else
