Tool menu updates:
- Improved ComboButton hierarchy support
- Fixed log define
- Downgrade command not mapped error log to verbose (asset editors currently need to register commands before they are mapped)
#rnx

diff --git a/Engine/Source/Developer/ToolMenus/Private/ToolMenus.cpp b/Engine/Source/Developer/ToolMenus/Private/ToolMenus.cpp
--- a/Engine/Source/Developer/ToolMenus/Private/ToolMenus.cpp
+++ b/Engine/Source/Developer/ToolMenus/Private/ToolMenus.cpp
@@ -14,6 +14,8 @@
 
 #define LOCTEXT_NAMESPACE "ToolMenuSubsystem"
 
+DEFINE_LOG_CATEGORY(LogToolMenus);
+
 UToolMenus* UToolMenus::Singleton = nullptr;
 bool UToolMenus::bHasShutDown = false;
 FSimpleMulticastDelegate UToolMenus::StartupCallbacks;
@@ -709,9 +711,14 @@ void UToolMenus::PopulateSubMenu(FMenuBuilder& MenuBuilder, TWeakObjectPtr<UTool
 	}
 }
 
-TSharedRef<SWidget> UToolMenus::GenerateToolbarComboButtonMenu(const FName SubMenuFullName, FToolMenuContext InContext)
+TSharedRef<SWidget> UToolMenus::GenerateToolbarComboButtonMenu(TWeakObjectPtr<UToolMenu> InParent, const FName InBlockName)
 {
-	return GenerateWidget(SubMenuFullName, InContext);
+	if (UToolMenu* GeneratedMenu = GenerateSubMenu(InParent.Get(), InBlockName))
+	{
+		return GenerateWidget(GeneratedMenu);
+	}
+
+	return SNullWidget::NullWidget;
 }
 
 void UToolMenus::FillMenuBarDropDown(class FMenuBuilder& MenuBuilder, FName InParentName, FName InChildName, FToolMenuContext InMenuContext)
@@ -977,7 +984,7 @@ void UToolMenus::PopulateToolBarBuilder(FToolBarBuilder& ToolBarBuilder, UToolMe
 					}
 					else
 					{
-						UE_LOG(LogToolMenus, Error, TEXT("UI command not found for toolbar entry: %s, toolbar: %s"), *Block.Name.ToString(), *MenuData->MenuName.ToString());
+						UE_LOG(LogToolMenus, Verbose, TEXT("UI command not found for toolbar entry: %s, toolbar: %s"), *Block.Name.ToString(), *MenuData->MenuName.ToString());
 					}
 
 					ToolBarBuilder.AddToolBarButton(Block.Command, Block.Name, Block.Label, Block.ToolTip, Block.Icon, Block.TutorialHighlightName);
@@ -1008,7 +1015,7 @@ void UToolMenus::PopulateToolBarBuilder(FToolBarBuilder& ToolBarBuilder, UToolMe
 				else
 				{
 					FName SubMenuFullName = JoinMenuPaths(MenuData->MenuName, Block.Name);
-					FOnGetContent Delegate = FOnGetContent::CreateUObject(this, &UToolMenus::GenerateToolbarComboButtonMenu, SubMenuFullName, MenuData->Context);
+					FOnGetContent Delegate = FOnGetContent::CreateUObject(this, &UToolMenus::GenerateToolbarComboButtonMenu, TWeakObjectPtr<UToolMenu>(MenuData), Block.Name);
 					ToolBarBuilder.AddComboButton(UIAction, Delegate, Block.Label, Block.ToolTip, Block.Icon, Block.ToolBarData.bSimpleComboBox, Block.TutorialHighlightName);
 				}
 			}
diff --git a/Engine/Source/Developer/ToolMenus/Public/IToolMenusModule.h b/Engine/Source/Developer/ToolMenus/Public/IToolMenusModule.h
--- a/Engine/Source/Developer/ToolMenus/Public/IToolMenusModule.h
+++ b/Engine/Source/Developer/ToolMenus/Public/IToolMenusModule.h
@@ -5,7 +5,7 @@
 #include "Modules/ModuleInterface.h"
 #include "Modules/ModuleManager.h"
 
-DEFINE_LOG_CATEGORY_STATIC(LogToolMenus, Log, All);
+DECLARE_LOG_CATEGORY_EXTERN(LogToolMenus, Log, All);
 
 class IToolMenusModule : public IModuleInterface
 {
diff --git a/Engine/Source/Developer/ToolMenus/Public/ToolMenus.h b/Engine/Source/Developer/ToolMenus/Public/ToolMenus.h
--- a/Engine/Source/Developer/ToolMenus/Public/ToolMenus.h
+++ b/Engine/Source/Developer/ToolMenus/Public/ToolMenus.h
@@ -322,7 +322,7 @@ private:
 	void PopulateMenuBarBuilder(FMenuBarBuilder& MenuBarBuilder, UToolMenu* MenuData);
 	void PopulateToolBarBuilder(FToolBarBuilder& ToolBarBuilder, UToolMenu* MenuData);
 
-	TSharedRef<SWidget> GenerateToolbarComboButtonMenu(const FName SubMenuFullName, FToolMenuContext InContext);
+	TSharedRef<SWidget> GenerateToolbarComboButtonMenu(TWeakObjectPtr<UToolMenu> InParent, const FName InBlockName);
 
 	FOnGetContent ConvertWidgetChoice(const FNewToolMenuWidgetChoice& Choice, const FToolMenuContext& Context) const;
 
