Tool menu entries added before a menu registered now appear after all others
#rnx

diff --git a/Engine/Source/Developer/ToolMenus/Private/ToolMenu.cpp b/Engine/Source/Developer/ToolMenus/Private/ToolMenu.cpp
--- a/Engine/Source/Developer/ToolMenus/Private/ToolMenu.cpp
+++ b/Engine/Source/Developer/ToolMenus/Private/ToolMenu.cpp
@@ -18,6 +18,7 @@ UToolMenu::UToolMenu() :
 	, bToolBarIsFocusable(false)
 	, bToolBarForceSmallIcons(false)
 	, bRegistered(false)
+	, bIsRegistering(false)
 	, StyleSet(&FCoreStyle::Get())
 {
 }
@@ -119,6 +120,11 @@ FToolMenuSection& UToolMenu::AddDynamicSection(const FName SectionName, const FN
 	return Section;
 }
 
+bool UToolMenu::IsRegistering() const
+{
+	return bIsRegistering;
+}
+
 FToolMenuSection& UToolMenu::AddSection(const FName SectionName, const TAttribute< FText >& InLabel, const FToolMenuInsert InPosition)
 {
 	for (FToolMenuSection& Section : Sections)
@@ -141,7 +147,7 @@ FToolMenuSection& UToolMenu::AddSection(const FName SectionName, const TAttribut
 
 	FToolMenuSection& NewSection = Sections.AddDefaulted_GetRef();
 	NewSection.InitSection(SectionName, InLabel, InPosition);
-	//NewSection.OwnerMenu = this;
+	NewSection.bIsRegistering = IsRegistering();
 	return NewSection;
 }
 
diff --git a/Engine/Source/Developer/ToolMenus/Private/ToolMenuEntry.cpp b/Engine/Source/Developer/ToolMenus/Private/ToolMenuEntry.cpp
--- a/Engine/Source/Developer/ToolMenus/Private/ToolMenuEntry.cpp
+++ b/Engine/Source/Developer/ToolMenus/Private/ToolMenuEntry.cpp
@@ -15,7 +15,8 @@ FToolMenuEntry::FToolMenuEntry() :
 	Type(EMultiBlockType::None),
 	UserInterfaceActionType(EUserInterfaceActionType::Button),
 	bShouldCloseWindowAfterMenuSelection(true),
-	ScriptObject(nullptr)
+	ScriptObject(nullptr),
+	bAddedDuringRegister(false)
 {
 }
 
@@ -25,7 +26,8 @@ FToolMenuEntry::FToolMenuEntry(const FToolMenuOwner InOwner, const FName InName,
 	Type(InType),
 	UserInterfaceActionType(EUserInterfaceActionType::Button),
 	bShouldCloseWindowAfterMenuSelection(true),
-	ScriptObject(nullptr)
+	ScriptObject(nullptr),
+	bAddedDuringRegister(false)
 {
 }
 
diff --git a/Engine/Source/Developer/ToolMenus/Private/ToolMenuSection.cpp b/Engine/Source/Developer/ToolMenus/Private/ToolMenuSection.cpp
--- a/Engine/Source/Developer/ToolMenus/Private/ToolMenuSection.cpp
+++ b/Engine/Source/Developer/ToolMenus/Private/ToolMenuSection.cpp
@@ -11,7 +11,8 @@
 #include "Internationalization/Internationalization.h"
 
 FToolMenuSection::FToolMenuSection() :
-	ToolMenuSectionDynamic(nullptr)
+	ToolMenuSectionDynamic(nullptr),
+	bIsRegistering(false)
 {
 
 }
@@ -33,22 +34,32 @@ void FToolMenuSection::InitGeneratedSectionCopy(const FToolMenuSection& Source,
 	Context = InContext;
 }
 
+bool FToolMenuSection::IsRegistering() const
+{
+	return bIsRegistering;
+}
+
 FToolMenuEntry& FToolMenuSection::AddEntry(const FToolMenuEntry& Args)
 {
 	if (Args.Name == NAME_None)
 	{
-		return Blocks.Add_GetRef(Args);
+		FToolMenuEntry& Result = Blocks.Add_GetRef(Args);
+		Result.bAddedDuringRegister = IsRegistering();
+		return Result;
 	}
 
 	int32 BlockIndex = IndexOfBlock(Args.Name);
 	if (BlockIndex != INDEX_NONE)
 	{
 		Blocks[BlockIndex] = Args;
+		Blocks[BlockIndex].bAddedDuringRegister = IsRegistering();
 		return Blocks[BlockIndex];
 	}
 	else
 	{
-		return Blocks.Add_GetRef(Args);
+		FToolMenuEntry& Result = Blocks.Add_GetRef(Args);
+		Result.bAddedDuringRegister = IsRegistering();
+		return Result;
 	}
 }
 
diff --git a/Engine/Source/Developer/ToolMenus/Private/ToolMenus.cpp b/Engine/Source/Developer/ToolMenus/Private/ToolMenus.cpp
--- a/Engine/Source/Developer/ToolMenus/Private/ToolMenus.cpp
+++ b/Engine/Source/Developer/ToolMenus/Private/ToolMenus.cpp
@@ -202,13 +202,21 @@ void UToolMenus::AssembleMenuSection(UToolMenu* GeneratedMenu, const UToolMenu*
 {
 	// Build list of blocks in expected order including blocks created by construct delegates
 	TArray<FToolMenuEntry> RemainingBlocks;
+	TArray<FToolMenuEntry> BlocksToAddLast;
 
 	UToolMenu* ConstructedEntries = nullptr;
 	for (const FToolMenuEntry& Block : OtherSection.Blocks)
 	{
 		if (!Block.IsNonLegacyDynamicConstruct())
 		{
-			RemainingBlocks.Add(Block);
+			if (Block.bAddedDuringRegister)
+			{
+				RemainingBlocks.Add(Block);
+			}
+			else
+			{
+				BlocksToAddLast.Add(Block);
+			}
 			continue;
 		}
 
@@ -264,7 +272,14 @@ void UToolMenus::AssembleMenuSection(UToolMenu* GeneratedMenu, const UToolMenu*
 			}
 			else
 			{
-				RemainingBlocks.Add(GeneratedEntry);
+				if (GeneratedEntry.bAddedDuringRegister)
+				{
+					RemainingBlocks.Add(GeneratedEntry);
+				}
+				else
+				{
+					BlocksToAddLast.Add(GeneratedEntry);
+				}
 				GeneratedEntries.RemoveAt(0, 1, false);
 			}
 		}
@@ -276,6 +291,8 @@ void UToolMenus::AssembleMenuSection(UToolMenu* GeneratedMenu, const UToolMenu*
 		ConstructedEntries = nullptr;
 	}
 
+	RemainingBlocks.Append(BlocksToAddLast);
+
 	// Repeatedly loop because insert location may not exist until later in list
 	while (RemainingBlocks.Num() > 0)
 	{
@@ -1674,6 +1691,11 @@ UToolMenu* UToolMenus::RegisterMenu(const FName InName, const FName InParent, EM
 			Found->MenuType = InType;
 			Found->MenuOwner = CurrentOwner();
 			Found->bRegistered = true;
+			Found->bIsRegistering = true;
+			for (FToolMenuSection& Section : Found->Sections)
+			{
+				Section.bIsRegistering = Found->bIsRegistering;
+			}
 		}
 		else if (bWarnIfAlreadyRegistered)
 		{
@@ -1686,6 +1708,7 @@ UToolMenu* UToolMenus::RegisterMenu(const FName InName, const FName InParent, EM
 	UToolMenu* ToolMenu = NewObject<UToolMenu>(this);
 	ToolMenu->InitMenu(CurrentOwner(), InName, InParent, InType);
 	ToolMenu->bRegistered = true;
+	ToolMenu->bIsRegistering = true;
 	Menus.Add(InName, ToolMenu);
 	return ToolMenu;
 }
@@ -1694,12 +1717,18 @@ UToolMenu* UToolMenus::ExtendMenu(const FName InName)
 {
 	if (UToolMenu* Found = FindMenu(InName))
 	{
+		Found->bIsRegistering = false;
+		for (FToolMenuSection& Section : Found->Sections)
+		{
+			Section.bIsRegistering = Found->bIsRegistering;
+		}
 		return Found;
 	}
 
 	UToolMenu* ToolMenu = NewObject<UToolMenu>(this);
 	ToolMenu->MenuName = InName;
 	ToolMenu->bRegistered = false;
+	ToolMenu->bIsRegistering = false;
 	Menus.Add(InName, ToolMenu);
 	return ToolMenu;
 }
diff --git a/Engine/Source/Developer/ToolMenus/Public/ToolMenu.h b/Engine/Source/Developer/ToolMenus/Public/ToolMenu.h
--- a/Engine/Source/Developer/ToolMenus/Public/ToolMenu.h
+++ b/Engine/Source/Developer/ToolMenus/Public/ToolMenu.h
@@ -101,6 +101,8 @@ private:
 
 	int32 FindInsertIndex(const FToolMenuSection& InSection) const;
 
+	bool IsRegistering() const;
+
 public:
 
 	UPROPERTY(EditAnywhere, BlueprintReadWrite, Category = "Tool Menus")
@@ -157,6 +159,7 @@ public:
 private:
 
 	bool bRegistered;
+	bool bIsRegistering;
 
 	const ISlateStyle* StyleSet;
 };
diff --git a/Engine/Source/Developer/ToolMenus/Public/ToolMenuEntry.h b/Engine/Source/Developer/ToolMenus/Public/ToolMenuEntry.h
--- a/Engine/Source/Developer/ToolMenus/Public/ToolMenuEntry.h
+++ b/Engine/Source/Developer/ToolMenus/Public/ToolMenuEntry.h
@@ -176,5 +176,7 @@ private:
 
 	FNewToolMenuSectionDelegate Construct;
 	FNewToolMenuDelegateLegacy ConstructLegacy;
+
+	bool bAddedDuringRegister;
 };
 
diff --git a/Engine/Source/Developer/ToolMenus/Public/ToolMenuSection.h b/Engine/Source/Developer/ToolMenus/Public/ToolMenuSection.h
--- a/Engine/Source/Developer/ToolMenus/Public/ToolMenuSection.h
+++ b/Engine/Source/Developer/ToolMenus/Public/ToolMenuSection.h
@@ -69,6 +69,8 @@ private:
 
 	bool IsNonLegacyDynamic() const;
 
+	bool IsRegistering() const;
+
 	friend class UToolMenuSectionExtensions;
 	friend class UToolMenus;
 	friend class UToolMenu;
@@ -93,4 +95,8 @@ public:
 	TAttribute<FText> Label;
 
 	FNewSectionConstructChoice Construct;
+
+private:
+
+	bool bIsRegistering;
 };
