Fixed submenus not being passed command lists.
#rnx

diff --git a/Engine/Source/Developer/ToolMenus/Private/ToolMenus.cpp b/Engine/Source/Developer/ToolMenus/Private/ToolMenus.cpp
--- a/Engine/Source/Developer/ToolMenus/Private/ToolMenus.cpp
+++ b/Engine/Source/Developer/ToolMenus/Private/ToolMenus.cpp
@@ -856,17 +856,32 @@ void UToolMenus::PopulateMenuBuilder(FMenuBuilder& MenuBuilder, UToolMenu* MenuD
 					}
 					else
 					{
-						MenuBuilder.AddSubMenu(
-							Block.Label,
-							Block.ToolTip,
-							NewMenuDelegate,
-							UIAction,
-							Block.Name,
-							Block.UserInterfaceActionType,
-							Block.SubMenuData.bOpenSubMenuOnClick,
-							Block.Icon.Get(),
-							Block.bShouldCloseWindowAfterMenuSelection
+						if (UIAction.IsBound())
+						{
+							MenuBuilder.AddSubMenu(
+								Block.Label,
+								Block.ToolTip,
+								NewMenuDelegate,
+								UIAction,
+								Block.Name,
+								Block.UserInterfaceActionType,
+								Block.SubMenuData.bOpenSubMenuOnClick,
+								Block.Icon.Get(),
+								Block.bShouldCloseWindowAfterMenuSelection
 							);
+						}
+						else
+						{
+							MenuBuilder.AddSubMenu(
+								Block.Label,
+								Block.ToolTip,
+								NewMenuDelegate,
+								Block.SubMenuData.bOpenSubMenuOnClick,
+								Block.Icon.Get(),
+								Block.bShouldCloseWindowAfterMenuSelection,
+								Block.Name
+							);
+						}
 					}
 				}
 				else
