Tool Menus: Impoved support for substitution of root menus in a hierarchy

diff --git a/Engine/Source/Developer/ToolMenus/Private/ToolMenu.cpp b/Engine/Source/Developer/ToolMenus/Private/ToolMenu.cpp
--- a/Engine/Source/Developer/ToolMenus/Private/ToolMenu.cpp
+++ b/Engine/Source/Developer/ToolMenus/Private/ToolMenu.cpp
@@ -328,10 +328,14 @@ TArray<FName> UToolMenu::GetMenuHierarchyNames(bool bIncludeSubMenuRoot) const
 {
 	TArray<FName> HierarchyNames;
 
-	TArray<UToolMenu*> Hierarchy = UToolMenus::Get()->CollectHierarchy(GetMenuName());
-	for (int32 i = Hierarchy.Num() - 1; i >= 0; --i)
+	TArray<UToolMenu*> Hierarchy;
+	if (UToolMenus::Get()->FindMenu(GetMenuName()) != nullptr)
 	{
-		HierarchyNames.AddUnique(Hierarchy[i]->GetMenuName());
+		Hierarchy = UToolMenus::Get()->CollectHierarchy(GetMenuName());
+		for (int32 i = Hierarchy.Num() - 1; i >= 0; --i)
+		{
+			HierarchyNames.AddUnique(Hierarchy[i]->GetMenuName());
+		}
 	}
 
 	if (bIncludeSubMenuRoot && SubMenuParent)
diff --git a/Engine/Source/Developer/ToolMenus/Private/ToolMenus.cpp b/Engine/Source/Developer/ToolMenus/Private/ToolMenus.cpp
--- a/Engine/Source/Developer/ToolMenus/Private/ToolMenus.cpp
+++ b/Engine/Source/Developer/ToolMenus/Private/ToolMenus.cpp
@@ -160,30 +160,54 @@ bool UToolMenus::IsMenuRegistered(const FName Name) const
 TArray<UToolMenu*> UToolMenus::CollectHierarchy(const FName InName)
 {
 	TArray<UToolMenu*> Result;
+	TArray<FName> SubstitutedMenuNames;
 
-	UToolMenu* Current = FindMenu(InName);
-	while (Current)
+	FName CurrentMenuName = InName;
+	while (CurrentMenuName != NAME_None)
 	{
-		// Detect infinite loop
-		for (UToolMenu* Other : Result)
+		FName AdjustedMenuName = CurrentMenuName;
+		if (!SubstitutedMenuNames.Contains(AdjustedMenuName))
 		{
-			if (Other->MenuName == Current->MenuName)
+			if (const FName* SubstitutionName = MenuSubstitutionsDuringGenerate.Find(CurrentMenuName))
 			{
-				UE_LOG(LogToolMenus, Warning, TEXT("Infinite loop detected in tool menu: %s"), *InName.ToString());
-				return TArray<UToolMenu*>();
+				// Allow collection hierarchy when InName is a substitute for one of InName's parents
+				// Will occur in menu editor when a substitute menu is selected from drop down list
+				bool bSubstituteAlreadyInHierarchy = false;
+				for (const UToolMenu* Other : Result)
+				{
+					if (Other->GetMenuName() == *SubstitutionName)
+					{
+						bSubstituteAlreadyInHierarchy = true;
+						break;
+					}
+				}
+
+				if (!bSubstituteAlreadyInHierarchy)
+				{
+					AdjustedMenuName = *SubstitutionName;
+
+					// Handle substitute's parent hierarchy including the original menu again by not substituting the same menu twice
+					SubstitutedMenuNames.Add(CurrentMenuName);
+				}
 			}
 		}
 
-		Result.Add(Current);
-
-		if (Current->MenuParent != NAME_None)
+		UToolMenu* Current = FindMenu(AdjustedMenuName);
+		if (!Current)
 		{
-			Current = FindMenu(Current->MenuParent);
+			UE_LOG(LogToolMenus, Warning, TEXT("Failed to find menu: %s for %s"), *AdjustedMenuName.ToString(), *InName.ToString());
+			return TArray<UToolMenu*>();
 		}
-		else
+
+		if (Result.Contains(Current))
 		{
-			break;
+			UE_LOG(LogToolMenus, Warning, TEXT("Infinite loop detected in tool menu: %s"), *InName.ToString());
+			return TArray<UToolMenu*>();
 		}
+
+		Result.Add(Current);
+
+		CurrentMenuName = Current->MenuParent;
 	}
 
 	Algo::Reverse(Result);
@@ -682,10 +706,6 @@ UToolMenu* UToolMenus::GenerateSubMenu(const UToolMenu* InGeneratedParent, const
 	}
 
 	FName SubMenuFullName = JoinMenuPaths(InGeneratedParent->GetMenuName(), InBlockName);
-	if (const FName* ReplacementName = MenuSubstitutionsDuringGenerate.Find(SubMenuFullName))
-	{
-		return GenerateMenuFromHierarchy(CollectHierarchy(*ReplacementName), InGeneratedParent->Context);
-	}
 
 	const FToolMenuEntry* Block = InGeneratedParent->FindEntry(InBlockName);
 	if (!Block)
@@ -1510,14 +1530,7 @@ void UToolMenus::RemoveSubstitutionDuringGenerate(const FName InMenu)
 
 UToolMenu* UToolMenus::GenerateMenu(const FName Name, const FToolMenuContext& InMenuContext)
 {
-	// Allow substituting one menu for another
-	FName MenuName = Name;
-	if (const FName* Found = MenuSubstitutionsDuringGenerate.Find(Name))
-	{
-		MenuName = *Found;
-	}
-
-	return GenerateMenuFromHierarchy(CollectHierarchy(MenuName), InMenuContext);
+	return GenerateMenuFromHierarchy(CollectHierarchy(Name), InMenuContext);
 }
 
 UToolMenu* UToolMenus::GenerateMenuFromHierarchy(const TArray<UToolMenu*>& Hierarchy, const FToolMenuContext& InMenuContext)
