Tool Menus - Add support for submenus that open to display a custom widget.
- Menu bar insert at specific location logic fixed.
#rnx

diff --git a/Engine/Source/Developer/ToolMenus/Private/ToolMenuContext.cpp b/Engine/Source/Developer/ToolMenus/Private/ToolMenuContext.cpp
--- a/Engine/Source/Developer/ToolMenus/Private/ToolMenuContext.cpp
+++ b/Engine/Source/Developer/ToolMenus/Private/ToolMenuContext.cpp
@@ -119,7 +119,7 @@ void FToolMenuContext::AddExtender(const TSharedPtr<FExtender>& InExtender)
 	Extenders.AddUnique(InExtender);
 }
 
-TSharedPtr<FExtender> FToolMenuContext::GetAllExtenders()
+TSharedPtr<FExtender> FToolMenuContext::GetAllExtenders() const
 {
 	return FExtender::Combine(Extenders);
 }
diff --git a/Engine/Source/Developer/ToolMenus/Private/ToolMenuEntry.cpp b/Engine/Source/Developer/ToolMenus/Private/ToolMenuEntry.cpp
--- a/Engine/Source/Developer/ToolMenus/Private/ToolMenuEntry.cpp
+++ b/Engine/Source/Developer/ToolMenus/Private/ToolMenuEntry.cpp
@@ -165,7 +165,7 @@ FToolMenuEntry FToolMenuEntry::InitToolBarButton(const TSharedPtr< const FUIComm
 	return Entry;
 }
 
-FToolMenuEntry FToolMenuEntry::InitComboButton(const FName InName, const FToolUIActionChoice& InAction, const FNewToolMenuWidgetChoice& InMenuContentGenerator, const TAttribute<FText>& InLabel, const TAttribute<FText>& InToolTip, const TAttribute<FSlateIcon>& InIcon, bool bInSimpleComboBox, FName InTutorialHighlightName)
+FToolMenuEntry FToolMenuEntry::InitComboButton(const FName InName, const FToolUIActionChoice& InAction, const FNewToolMenuChoice& InMenuContentGenerator, const TAttribute<FText>& InLabel, const TAttribute<FText>& InToolTip, const TAttribute<FSlateIcon>& InIcon, bool bInSimpleComboBox, FName InTutorialHighlightName)
 {
 	FToolMenuEntry Entry(UToolMenus::Get()->CurrentOwner(), InName, EMultiBlockType::ToolBarComboButton);
 	Entry.TutorialHighlightName = InTutorialHighlightName;
diff --git a/Engine/Source/Developer/ToolMenus/Private/ToolMenuSection.cpp b/Engine/Source/Developer/ToolMenus/Private/ToolMenuSection.cpp
--- a/Engine/Source/Developer/ToolMenus/Private/ToolMenuSection.cpp
+++ b/Engine/Source/Developer/ToolMenus/Private/ToolMenuSection.cpp
@@ -173,38 +173,6 @@ bool FToolMenuSection::IsNonLegacyDynamic() const
 	return ToolMenuSectionDynamic || Construct.NewToolMenuDelegate.IsBound();
 }
 
-void FToolMenuSection::AssembleBlock(const FToolMenuEntry& Block)
-{
-	const EToolMenuInsertType Position = Block.InsertPosition.Position;
-
-	int32 ExistingIndex = IndexOfBlock(Block.Name);
-	if (ExistingIndex != INDEX_NONE)
-	{
-		Blocks[ExistingIndex] = Block;
-	}
-	else if (Position == EToolMenuInsertType::Before || Position == EToolMenuInsertType::After)
-	{
-		int32 DestIndex = IndexOfBlock(Block.InsertPosition.Name);
-		if (DestIndex != INDEX_NONE)
-		{
-			if (Position == EToolMenuInsertType::After)
-			{
-				++DestIndex;
-			}
-
-			Blocks.Insert(Block, DestIndex);
-		}
-	}
-	else if (Position == EToolMenuInsertType::First)
-	{
-		Blocks.Insert(Block, 0);
-	}
-	else
-	{
-		Blocks.Add(Block);
-	}
-}
-
 int32 FToolMenuSection::RemoveEntry(const FName InName)
 {
 	return Blocks.RemoveAll([InName](const FToolMenuEntry& Block) { return Block.Name == InName; });
@@ -248,7 +216,7 @@ int32 FToolMenuSection::FindBlockInsertIndex(const FToolMenuEntry& InBlock) cons
 		return INDEX_NONE;
 	}
 
-	if (InsertPosition.Position == EToolMenuInsertType::After)
+	if (InPosition.Position == EToolMenuInsertType::After)
 	{
 		++DestIndex;
 	}
diff --git a/Engine/Source/Developer/ToolMenus/Private/ToolMenus.cpp b/Engine/Source/Developer/ToolMenus/Private/ToolMenus.cpp
--- a/Engine/Source/Developer/ToolMenus/Private/ToolMenus.cpp
+++ b/Engine/Source/Developer/ToolMenus/Private/ToolMenus.cpp
@@ -421,7 +421,13 @@ void UToolMenus::AssembleMenu(UToolMenu* GeneratedMenu, const UToolMenu* Other)
 		int32 NumHandled = 0;
 		for (int32 i=0; i < RemainingSections.Num(); ++i)
 		{
-			const FToolMenuSection& RemainingSection = RemainingSections[i];
+			FToolMenuSection& RemainingSection = RemainingSections[i];
+
+			// Menubars do not have sections, combine all sections into one
+			if (GeneratedMenu->MenuType == EMultiBoxType::MenuBar)
+			{
+				RemainingSection.Name = NAME_None;
+			}
 
 			// Update existing section
 			FToolMenuSection* Section = GeneratedMenu->FindSection(RemainingSection.Name);
@@ -667,32 +673,9 @@ void UToolMenus::ApplyCustomization(UToolMenu* GeneratedMenu)
 
 void UToolMenus::AssembleMenuHierarchy(UToolMenu* GeneratedMenu, const TArray<UToolMenu*>& Hierarchy)
 {
-	if (GeneratedMenu->MenuType == EMultiBoxType::MenuBar)
-	{
-		// Menu Bars require one section
-		if (GeneratedMenu->Sections.Num() == 0)
-		{
-			GeneratedMenu->Sections.AddDefaulted();
-		}
-
-		FToolMenuSection& MenuBarSection = GeneratedMenu->Sections[0];
-		for (const UToolMenu* MenuData : Hierarchy)
-		{
-			for (const FToolMenuSection& Section : MenuData->Sections)
-			{
-				for (const FToolMenuEntry& Block : Section.Blocks)
-				{
-					MenuBarSection.AssembleBlock(Block);
-				}
-			}
-		}
-	}
-	else
+	for (const UToolMenu* FoundParent : Hierarchy)
 	{
-		for (const UToolMenu* FoundParent : Hierarchy)
-		{
-			AssembleMenu(GeneratedMenu, FoundParent);
-		}
+		AssembleMenu(GeneratedMenu, FoundParent);
 	}
 
 	ApplyCustomization(GeneratedMenu);
@@ -734,7 +717,7 @@ UToolMenu* UToolMenus::GenerateSubMenu(const UToolMenu* InGeneratedParent, const
 	}
 
 	// Construct menu using delegate and insert as root so it can be overridden
-	if (Block->SubMenuData.ConstructMenu.NewToolMenuDelegate.IsBound())
+	if (Block->SubMenuData.ConstructMenu.NewToolMenu.IsBound())
 	{
 		UToolMenu* Menu = NewObject<UToolMenu>(this);
 		Menu->Context = InGeneratedParent->Context;
@@ -744,7 +727,7 @@ UToolMenu* UToolMenus::GenerateSubMenu(const UToolMenu* InGeneratedParent, const
 		Menu->SubMenuParent = InGeneratedParent;
 		Menu->SubMenuSourceEntryName = InBlockName;
 
-		Block->SubMenuData.ConstructMenu.NewToolMenuDelegate.Execute(Menu);
+		Block->SubMenuData.ConstructMenu.NewToolMenu.Execute(Menu);
 		Menu->MenuName = SubMenuFullName;
 		Hierarchy.Insert(Menu, 0);
 	}
@@ -896,16 +879,33 @@ void UToolMenus::PopulateMenuBuilder(FMenuBuilder& MenuBuilder, UToolMenu* MenuD
 				{
 					FName SubMenuFullName = JoinMenuPaths(MenuData->MenuName, Block.Name);
 					FNewMenuDelegate NewMenuDelegate;
-					if (Block.SubMenuData.ConstructMenu.NewMenuDelegate.IsBound())
+					bool bSubMenuAdded = false;
+
+					if (Block.SubMenuData.ConstructMenu.NewMenuLegacy.IsBound())
 					{
-						NewMenuDelegate = Block.SubMenuData.ConstructMenu.NewMenuDelegate;
+						NewMenuDelegate = Block.SubMenuData.ConstructMenu.NewMenuLegacy;
+					}
+					else if (Block.SubMenuData.ConstructMenu.NewToolMenuWidget.IsBound())
+					{
+						// Full replacement of the widget shown when submenu is opened
+						FOnGetContent OnGetContent = ConvertWidgetChoice(Block.SubMenuData.ConstructMenu, MenuData->Context);
+						if (OnGetContent.IsBound())
+						{
+							MenuBuilder.AddWrapperSubMenu(
+								Block.Label.Get(),
+								Block.ToolTip.Get(),
+								OnGetContent,
+								Block.Icon.Get()
+							);
+						}
+						bSubMenuAdded = true;
 					}
 					else if (Block.Name == NAME_None)
 					{
-						if (Block.SubMenuData.ConstructMenu.NewToolMenuDelegate.IsBound())
+						if (Block.SubMenuData.ConstructMenu.NewToolMenu.IsBound())
 						{
 							// Blocks with no name cannot call PopulateSubMenu()
-							NewMenuDelegate = FNewMenuDelegate::CreateUObject(this, &UToolMenus::PopulateSubMenuWithoutName, TWeakObjectPtr<UToolMenu>(MenuData), Block.SubMenuData.ConstructMenu.NewToolMenuDelegate);
+							NewMenuDelegate = FNewMenuDelegate::CreateUObject(this, &UToolMenus::PopulateSubMenuWithoutName, TWeakObjectPtr<UToolMenu>(MenuData), Block.SubMenuData.ConstructMenu.NewToolMenu);
 						}
 						else
 						{
@@ -917,44 +917,47 @@ void UToolMenus::PopulateMenuBuilder(FMenuBuilder& MenuBuilder, UToolMenu* MenuD
 						NewMenuDelegate = FNewMenuDelegate::CreateUObject(this, &UToolMenus::PopulateSubMenu, TWeakObjectPtr<UToolMenu>(MenuData), Block.Name);
 					}
 
-					if (Widget.IsValid())
+					if (!bSubMenuAdded)
 					{
-						if (UIAction.IsBound())
+						if (Widget.IsValid())
 						{
-							MenuBuilder.AddSubMenu(UIAction, Widget.ToSharedRef(), NewMenuDelegate, Block.bShouldCloseWindowAfterMenuSelection);
+							if (UIAction.IsBound())
+							{
+								MenuBuilder.AddSubMenu(UIAction, Widget.ToSharedRef(), NewMenuDelegate, Block.bShouldCloseWindowAfterMenuSelection);
+							}
+							else
+							{
+								MenuBuilder.AddSubMenu(Widget.ToSharedRef(), NewMenuDelegate, Block.SubMenuData.bOpenSubMenuOnClick, Block.bShouldCloseWindowAfterMenuSelection);
+							}
 						}
 						else
 						{
-							MenuBuilder.AddSubMenu(Widget.ToSharedRef(), NewMenuDelegate, Block.SubMenuData.bOpenSubMenuOnClick, Block.bShouldCloseWindowAfterMenuSelection);
-						}
-					}
-					else
-					{
-						if (UIAction.IsBound())
-						{
-							MenuBuilder.AddSubMenu(
-								Block.Label,
-								Block.ToolTip,
-								NewMenuDelegate,
-								UIAction,
-								Block.Name,
-								Block.UserInterfaceActionType,
-								Block.SubMenuData.bOpenSubMenuOnClick,
-								Block.Icon.Get(),
-								Block.bShouldCloseWindowAfterMenuSelection
-							);
-						}
-						else
-						{
-							MenuBuilder.AddSubMenu(
-								Block.Label,
-								Block.ToolTip,
-								NewMenuDelegate,
-								Block.SubMenuData.bOpenSubMenuOnClick,
-								Block.Icon.Get(),
-								Block.bShouldCloseWindowAfterMenuSelection,
-								Block.Name
-							);
+							if (UIAction.IsBound())
+							{
+								MenuBuilder.AddSubMenu(
+									Block.Label,
+									Block.ToolTip,
+									NewMenuDelegate,
+									UIAction,
+									Block.Name,
+									Block.UserInterfaceActionType,
+									Block.SubMenuData.bOpenSubMenuOnClick,
+									Block.Icon.Get(),
+									Block.bShouldCloseWindowAfterMenuSelection
+								);
+							}
+							else
+							{
+								MenuBuilder.AddSubMenu(
+									Block.Label,
+									Block.ToolTip,
+									NewMenuDelegate,
+									Block.SubMenuData.bOpenSubMenuOnClick,
+									Block.Icon.Get(),
+									Block.bShouldCloseWindowAfterMenuSelection,
+									Block.Name
+								);
+							}
 						}
 					}
 				}
@@ -1154,9 +1157,9 @@ void UToolMenus::PopulateMenuBarBuilder(FMenuBarBuilder& MenuBarBuilder, UToolMe
 		{
 			FName SubMenuFullName = JoinMenuPaths(MenuData->MenuName, Block.Name);
 			FNewMenuDelegate NewMenuDelegate;
-			if (Block.SubMenuData.ConstructMenu.NewMenuDelegate.IsBound())
+			if (Block.SubMenuData.ConstructMenu.NewMenuLegacy.IsBound())
 			{
-				NewMenuDelegate = Block.SubMenuData.ConstructMenu.NewMenuDelegate;
+				NewMenuDelegate = Block.SubMenuData.ConstructMenu.NewMenuLegacy;
 			}
 			else
 			{
@@ -1175,7 +1178,7 @@ void UToolMenus::PopulateMenuBarBuilder(FMenuBarBuilder& MenuBarBuilder, UToolMe
 	AddReferencedContextObjects(MenuBarBuilder.GetMultiBox(), MenuData);
 }
 
-FOnGetContent UToolMenus::ConvertWidgetChoice(const FNewToolMenuWidgetChoice& Choice, const FToolMenuContext& Context) const
+FOnGetContent UToolMenus::ConvertWidgetChoice(const FNewToolMenuChoice& Choice, const FToolMenuContext& Context) const
 {
 	if (Choice.NewToolMenuWidget.IsBound())
 	{
@@ -1204,6 +1207,20 @@ FOnGetContent UToolMenus::ConvertWidgetChoice(const FNewToolMenuWidgetChoice& Ch
 			return SNullWidget::NullWidget;
 		});
 	}
+	else if (Choice.NewMenuLegacy.IsBound())
+	{
+		return FOnGetContent::CreateLambda([ToCall = Choice.NewMenuLegacy, Context]()
+		{
+			if (ToCall.IsBound())
+			{
+				FMenuBuilder MenuBuilder(true, Context.CommandList, Context.GetAllExtenders());
+				ToCall.Execute(MenuBuilder);
+				return MenuBuilder.MakeWidget();
+			}
+
+			return SNullWidget::NullWidget;
+		});
+	}
 	return Choice.OnGetContent;
 }
 
diff --git a/Engine/Source/Developer/ToolMenus/Public/ToolMenuContext.h b/Engine/Source/Developer/ToolMenus/Public/ToolMenuContext.h
--- a/Engine/Source/Developer/ToolMenus/Public/ToolMenuContext.h
+++ b/Engine/Source/Developer/ToolMenus/Public/ToolMenuContext.h
@@ -73,7 +73,7 @@ public:
 	const FUIAction* GetActionForCommand(TSharedPtr<const FUICommandInfo> Command) const;
 
 	void AddExtender(const TSharedPtr<FExtender>& InExtender);
-	TSharedPtr<FExtender> GetAllExtenders();
+	TSharedPtr<FExtender> GetAllExtenders() const;
 	void ReplaceExtenders(const TSharedPtr<FExtender>& InExtender);
 	void ResetExtenders();
 
diff --git a/Engine/Source/Developer/ToolMenus/Public/ToolMenuDelegates.h b/Engine/Source/Developer/ToolMenus/Public/ToolMenuDelegates.h
--- a/Engine/Source/Developer/ToolMenus/Public/ToolMenuDelegates.h
+++ b/Engine/Source/Developer/ToolMenus/Public/ToolMenuDelegates.h
@@ -68,17 +68,18 @@ struct TOOLMENUS_API FToolDynamicUIAction
 	FToolMenuDynamicIsActionButtonVisible IsActionVisibleDelegate;
 };
 
-struct TOOLMENUS_API FNewToolMenuWidgetChoice
+struct TOOLMENUS_API FNewToolMenuChoice
 {
-public:
-	FNewToolMenuWidgetChoice() {}
-	FNewToolMenuWidgetChoice(const FOnGetContent& InOnGetContent) : OnGetContent(InOnGetContent) {}
-	FNewToolMenuWidgetChoice(const FNewToolMenuWidget& InNewToolMenuWidget) : NewToolMenuWidget(InNewToolMenuWidget) {}
-	FNewToolMenuWidgetChoice(const FNewToolMenuDelegate& InNewToolMenu) : NewToolMenu(InNewToolMenu) {}
+	FNewToolMenuChoice() {}
+	FNewToolMenuChoice(const FOnGetContent& InOnGetContent) : OnGetContent(InOnGetContent) {}
+	FNewToolMenuChoice(const FNewToolMenuWidget& InNewToolMenuWidget) : NewToolMenuWidget(InNewToolMenuWidget) {}
+	FNewToolMenuChoice(const FNewToolMenuDelegate& InNewToolMenu) : NewToolMenu(InNewToolMenu) {}
+	FNewToolMenuChoice(const FNewMenuDelegate& InNewMenuLegacy) : NewMenuLegacy(InNewMenuLegacy) {}
 
 	FOnGetContent OnGetContent;
 	FNewToolMenuWidget NewToolMenuWidget;
 	FNewToolMenuDelegate NewToolMenu;
+	FNewMenuDelegate NewMenuLegacy;
 };
 
 struct TOOLMENUS_API FToolUIActionChoice
@@ -113,16 +114,6 @@ private:
 	TOptional<FToolDynamicUIAction> DynamicToolAction;
 };
 
-struct TOOLMENUS_API FNewToolMenuChoice
-{
-	FNewToolMenuChoice() {};
-	FNewToolMenuChoice(const FNewToolMenuDelegate& InNewToolMenuDelegate) : NewToolMenuDelegate(InNewToolMenuDelegate) {}
-	FNewToolMenuChoice(const FNewMenuDelegate& InNewMenuDelegate) : NewMenuDelegate(InNewMenuDelegate) {}
-
-	FNewToolMenuDelegate NewToolMenuDelegate;
-	FNewMenuDelegate NewMenuDelegate;
-};
-
 struct TOOLMENUS_API FNewSectionConstructChoice
 {
 	FNewSectionConstructChoice() {};
diff --git a/Engine/Source/Developer/ToolMenus/Public/ToolMenuEntry.h b/Engine/Source/Developer/ToolMenus/Public/ToolMenuEntry.h
--- a/Engine/Source/Developer/ToolMenus/Public/ToolMenuEntry.h
+++ b/Engine/Source/Developer/ToolMenus/Public/ToolMenuEntry.h
@@ -53,7 +53,7 @@ public:
 	bool bForceSmallIcons;
 
 	/** Delegate that generates a widget for this combo button's menu content.  Called when the menu is summoned. */
-	FNewToolMenuWidgetChoice ComboButtonContextMenuGenerator;
+	FNewToolMenuChoice ComboButtonContextMenuGenerator;
 
 	/** Legacy delegate that generates a widget for this combo button's menu content.  Called when the menu is summoned. */
 	FNewToolBarDelegateLegacy ConstructLegacy;
@@ -95,7 +95,7 @@ struct TOOLMENUS_API FToolMenuEntry
 
 	static FToolMenuEntry InitToolBarButton(const FName InName, const FToolUIActionChoice& InAction, const TAttribute<FText>& InLabel = TAttribute<FText>(), const TAttribute<FText>& InToolTip = TAttribute<FText>(), const TAttribute<FSlateIcon>& InIcon = TAttribute<FSlateIcon>(), const EUserInterfaceActionType UserInterfaceActionType = EUserInterfaceActionType::Button, FName InTutorialHighlightName = NAME_None);
 	static FToolMenuEntry InitToolBarButton(const TSharedPtr< const FUICommandInfo >& InCommand, const TAttribute<FText>& InLabelOverride = TAttribute<FText>(), const TAttribute<FText>& InToolTipOverride = TAttribute<FText>(), const TAttribute<FSlateIcon>& InIconOverride = TAttribute<FSlateIcon>(), FName InTutorialHighlightName = NAME_None, const FName InNameOverride = NAME_None);
-	static FToolMenuEntry InitComboButton(const FName InName, const FToolUIActionChoice& InAction, const FNewToolMenuWidgetChoice& InMenuContentGenerator, const TAttribute<FText>& InLabelOverride = TAttribute<FText>(), const TAttribute<FText>& InToolTipOverride = TAttribute<FText>(), const TAttribute<FSlateIcon>& InIconOverride = TAttribute<FSlateIcon>(), bool bInSimpleComboBox = false, FName InTutorialHighlightName = NAME_None);
+	static FToolMenuEntry InitComboButton(const FName InName, const FToolUIActionChoice& InAction, const FNewToolMenuChoice& InMenuContentGenerator, const TAttribute<FText>& InLabelOverride = TAttribute<FText>(), const TAttribute<FText>& InToolTipOverride = TAttribute<FText>(), const TAttribute<FSlateIcon>& InIconOverride = TAttribute<FSlateIcon>(), bool bInSimpleComboBox = false, FName InTutorialHighlightName = NAME_None);
 
 	static FToolMenuEntry InitMenuSeparator(const FName InName);
 	static FToolMenuEntry InitToolBarSeparator(const FName InName);
diff --git a/Engine/Source/Developer/ToolMenus/Public/ToolMenuSection.h b/Engine/Source/Developer/ToolMenus/Public/ToolMenuSection.h
--- a/Engine/Source/Developer/ToolMenus/Public/ToolMenuSection.h
+++ b/Engine/Source/Developer/ToolMenus/Public/ToolMenuSection.h
@@ -65,8 +65,6 @@ private:
 	int32 IndexOfBlock(const FName InName) const;
 	int32 FindBlockInsertIndex(const FToolMenuEntry& InBlock) const;
 
-	void AssembleBlock(const FToolMenuEntry& InBlock);
-
 	bool IsNonLegacyDynamic() const;
 
 	bool IsRegistering() const;
diff --git a/Engine/Source/Developer/ToolMenus/Public/ToolMenus.h b/Engine/Source/Developer/ToolMenus/Public/ToolMenus.h
--- a/Engine/Source/Developer/ToolMenus/Public/ToolMenus.h
+++ b/Engine/Source/Developer/ToolMenus/Public/ToolMenus.h
@@ -329,7 +329,7 @@ private:
 
 	TSharedRef<SWidget> GenerateToolbarComboButtonMenu(TWeakObjectPtr<UToolMenu> InParent, const FName InBlockName);
 
-	FOnGetContent ConvertWidgetChoice(const FNewToolMenuWidgetChoice& Choice, const FToolMenuContext& Context) const;
+	FOnGetContent ConvertWidgetChoice(const FNewToolMenuChoice& Choice, const FToolMenuContext& Context) const;
 
 	/** Converts a string command to a FUIAction */
 	static FUIAction ConvertUIAction(const FToolMenuEntry& Block, const FToolMenuContext& Context);
diff --git a/Engine/Source/Editor/LevelEditor/Private/LevelEditorToolBar.cpp b/Engine/Source/Editor/LevelEditor/Private/LevelEditorToolBar.cpp
--- a/Engine/Source/Editor/LevelEditor/Private/LevelEditorToolBar.cpp
+++ b/Engine/Source/Editor/LevelEditor/Private/LevelEditorToolBar.cpp
@@ -1531,7 +1531,7 @@ void FLevelEditorToolBar::RegisterLevelEditorToolBar( const TSharedRef<FUIComman
 						FCanExecuteAction(),
 						FIsActionChecked(),
 						FIsActionButtonVisible::CreateStatic(FLevelEditorActionCallbacks::CanShowSourceCodeActions)),
-					FNewToolMenuWidgetChoice(),
+					FNewToolMenuChoice(),
 					LOCTEXT("CompileCombo_Label", "Compile Options"),
 					LOCTEXT("CompileComboToolTip", "Compile options menu"),
 					FSlateIcon(FEditorStyle::GetStyleSetName(), "LevelEditor.Recompile"),
diff --git a/Engine/Source/Editor/UnrealEd/Private/Toolkits/AssetEditorToolkit.cpp b/Engine/Source/Editor/UnrealEd/Private/Toolkits/AssetEditorToolkit.cpp
--- a/Engine/Source/Editor/UnrealEd/Private/Toolkits/AssetEditorToolkit.cpp
+++ b/Engine/Source/Editor/UnrealEd/Private/Toolkits/AssetEditorToolkit.cpp
@@ -983,6 +983,11 @@ FName FAssetEditorToolkit::GetToolMenuAppName() const
 	return GetToolkitFName();
 }
 
+FName FAssetEditorToolkit::GetToolMenuName() const
+{
+	return *(TEXT("AssetEditor.") + GetToolMenuAppName().ToString() + TEXT(".MainMenu"));
+}
+
 FName FAssetEditorToolkit::GetToolMenuToolbarName() const
 {
 	FName ParentName;
diff --git a/Engine/Source/Editor/UnrealEd/Public/Toolkits/AssetEditorToolkit.h b/Engine/Source/Editor/UnrealEd/Public/Toolkits/AssetEditorToolkit.h
--- a/Engine/Source/Editor/UnrealEd/Public/Toolkits/AssetEditorToolkit.h
+++ b/Engine/Source/Editor/UnrealEd/Public/Toolkits/AssetEditorToolkit.h
@@ -194,6 +194,7 @@ public:
 	virtual void InitToolMenuContext(FToolMenuContext& MenuContext);
 	FName GetToolMenuToolbarName() const;
 	FName GetToolMenuAppName() const;
+	FName GetToolMenuName() const;
 
 	/** Called at the end of RegenerateMenusAndToolbars() */
 	virtual void PostRegenerateMenusAndToolbars() { }
