Tool Menus: Allow extending submenus that are not registered
#rnx

diff --git a/Engine/Source/Developer/ToolMenus/Private/ToolMenus.cpp b/Engine/Source/Developer/ToolMenus/Private/ToolMenus.cpp
--- a/Engine/Source/Developer/ToolMenus/Private/ToolMenus.cpp
+++ b/Engine/Source/Developer/ToolMenus/Private/ToolMenus.cpp
@@ -157,7 +157,7 @@ bool UToolMenus::IsMenuRegistered(const FName Name) const
 	return Found && *Found && (*Found)->IsRegistered();
 }
 
-TArray<UToolMenu*> UToolMenus::CollectHierarchy(const FName InName)
+TArray<UToolMenu*> UToolMenus::CollectHierarchy(const FName InName, const TMap<FName, FName>& UnregisteredParentNames)
 {
 	TArray<UToolMenu*> Result;
 	TArray<FName> SubstitutedMenuNames;
@@ -207,7 +207,18 @@ TArray<UToolMenu*> UToolMenus::CollectHierarchy(const FName InName)
 
 		Result.Add(Current);
 
-		CurrentMenuName = Current->MenuParent;
+		if (Current->IsRegistered())
+		{
+			CurrentMenuName = Current->MenuParent;
+		}
+		else if (const FName* FoundUnregisteredParentName = UnregisteredParentNames.Find(CurrentMenuName))
+		{
+			CurrentMenuName = *FoundUnregisteredParentName;
+		}
+		else
+		{
+			CurrentMenuName = NAME_None;
+		}
 	}
 
 	Algo::Reverse(Result);
@@ -215,6 +226,12 @@ TArray<UToolMenu*> UToolMenus::CollectHierarchy(const FName InName)
 	return Result;
 }
 
+TArray<UToolMenu*> UToolMenus::CollectHierarchy(const FName InName)
+{
+	TMap<FName, FName> UnregisteredParents;
+	return CollectHierarchy(InName, UnregisteredParents);
+}
+
 void UToolMenus::ListAllParents(const FName InName, TArray<FName>& AllParents)
 {
 	for (const UToolMenu* Menu : CollectHierarchy(InName))
@@ -700,19 +717,60 @@ UToolMenu* UToolMenus::GenerateSubMenu(const UToolMenu* InGeneratedParent, const
 	// Submenus that are constructed by delegates can also be overridden by menus in the database
 	TArray<UToolMenu*> Hierarchy;
 	{
+		struct FMenuHierarchyInfo
+		{
+			FMenuHierarchyInfo() : BaseMenu(nullptr), SubMenu(nullptr) { }
+			FName BaseMenuName;
+			FName SubMenuName;
+			UToolMenu* BaseMenu;
+			UToolMenu* SubMenu;
+		};
+
+		TArray<FMenuHierarchyInfo> HierarchyInfos;
+		TArray<UToolMenu*> UnregisteredHierarchy;
+
 		// Walk up all parent menus trying to find a menu
 		FName BaseName = InGeneratedParent->GetMenuName();
 		while (BaseName != NAME_None)
 		{
-			FName JoinedName = JoinMenuPaths(BaseName, InBlockName);
-			if (UToolMenu* Found = FindMenu(JoinedName))
+			FMenuHierarchyInfo& Info = HierarchyInfos.AddDefaulted_GetRef();
+			Info.BaseMenuName = BaseName;
+			Info.BaseMenu = FindMenu(Info.BaseMenuName);
+			Info.SubMenuName = JoinMenuPaths(BaseName, InBlockName);
+			Info.SubMenu = FindMenu(Info.SubMenuName);
+
+			if (Info.SubMenu)
 			{
-				Hierarchy = CollectHierarchy(JoinedName);
-				break;
+				if (Info.SubMenu->IsRegistered())
+				{
+					if (UnregisteredHierarchy.Num() == 0)
+					{
+						Hierarchy = CollectHierarchy(Info.SubMenuName);
+					}
+					else
+					{
+						UnregisteredHierarchy.Add(Info.SubMenu);
+					}
+					break;
+				}
+				else
+				{
+					UnregisteredHierarchy.Add(Info.SubMenu);
+				}
 			}
 
-			UToolMenu* BaseData = FindMenu(BaseName);
-			BaseName = BaseData ? BaseData->MenuParent : NAME_None;
+			BaseName = Info.BaseMenu  ? Info.BaseMenu->MenuParent : NAME_None;
+		}
+
+		if (UnregisteredHierarchy.Num() > 0)
+		{
+			// Create lookup for UToolMenus that were extended but not registered
+			TMap<FName, FName> UnregisteredParentNames;
+			for (int32 i = 0; i < UnregisteredHierarchy.Num() - 1; ++i)
+			{
+				UnregisteredParentNames.Add(UnregisteredHierarchy[i]->GetMenuName(), UnregisteredHierarchy[i + 1]->GetMenuName());
+			}
+			Hierarchy = CollectHierarchy(UnregisteredHierarchy[0]->GetMenuName(), UnregisteredParentNames);
 		}
 	}
 
diff --git a/Engine/Source/Developer/ToolMenus/Public/ToolMenus.h b/Engine/Source/Developer/ToolMenus/Public/ToolMenus.h
--- a/Engine/Source/Developer/ToolMenus/Public/ToolMenus.h
+++ b/Engine/Source/Developer/ToolMenus/Public/ToolMenus.h
@@ -342,6 +342,7 @@ private:
 
 	void PopulateSubMenu(FMenuBuilder& Builder, TWeakObjectPtr<UToolMenu> InParent, const FName InBlockName);
 	void PopulateSubMenuWithoutName(FMenuBuilder& MenuBuilder, TWeakObjectPtr<UToolMenu> InParent, const FNewToolMenuDelegate InNewToolMenuDelegate);
+	TArray<UToolMenu*> CollectHierarchy(const FName Name, const TMap<FName, FName>& UnregisteredParentNames);
 
 	void ListAllParents(const FName Name, TArray<FName>& AllParents);
 
