Tool Menus: Sort sections that are added to a menu before registration
#rnx

diff --git a/Engine/Source/Developer/ToolMenus/Private/ToolMenu.cpp b/Engine/Source/Developer/ToolMenus/Private/ToolMenu.cpp
--- a/Engine/Source/Developer/ToolMenus/Private/ToolMenu.cpp
+++ b/Engine/Source/Developer/ToolMenus/Private/ToolMenu.cpp
@@ -132,27 +132,60 @@ bool UToolMenu::IsRegistering() const
 
 FToolMenuSection& UToolMenu::AddSection(const FName SectionName, const TAttribute< FText >& InLabel, const FToolMenuInsert InPosition)
 {
-	for (FToolMenuSection& Section : Sections)
+	int32 InsertIndex = (SectionName != NAME_None) ? IndexOfSection(SectionName) : INDEX_NONE;
+	if (InsertIndex != INDEX_NONE)
 	{
-		if (Section.Name == SectionName)
+		if (InLabel.IsSet())
+		{
+			Sections[InsertIndex].Label = InLabel;
+		}
+
+		if (InPosition.Name != NAME_None)
 		{
-			if (InLabel.IsSet())
+			Sections[InsertIndex].InsertPosition = InPosition;
+		}
+
+		// Sort registered sections to appear before unregistered
+		if (IsRegistering() && !Sections[InsertIndex].bAddedDuringRegister)
+		{
+			Sections[InsertIndex].bAddedDuringRegister = true;
+
+			for (int32 i = 0; i < InsertIndex; ++i)
 			{
-				Section.Label = InLabel;
+				if (!Sections[i].bAddedDuringRegister)
+				{
+					FToolMenuSection RemovedSection;
+					Swap(Sections[InsertIndex], RemovedSection);
+					Sections.Insert(MoveTempIfPossible(RemovedSection), i);
+					Sections.RemoveAt(InsertIndex + 1, 1, false);
+					InsertIndex = i;
+				}
 			}
+		}
+
+		return Sections[InsertIndex];
+	}
+	else
+	{
+		InsertIndex = Sections.Num();
+	}
 
-			if (InPosition.Name != NAME_None)
+	if (IsRegistering())
+	{
+		for (int32 i=0; i < Sections.Num(); ++i)
+		{
+			if (!Sections[i].bAddedDuringRegister)
 			{
-				Section.InsertPosition = InPosition;
+				InsertIndex = i;
+				break;
 			}
-
-			return Section;
 		}
 	}
 
-	FToolMenuSection& NewSection = Sections.AddDefaulted_GetRef();
+	FToolMenuSection& NewSection = Sections.InsertDefaulted_GetRef(InsertIndex);
 	NewSection.InitSection(SectionName, InLabel, InPosition);
 	NewSection.bIsRegistering = IsRegistering();
+	NewSection.bAddedDuringRegister = IsRegistering();
 	return NewSection;
 }
 
diff --git a/Engine/Source/Developer/ToolMenus/Private/ToolMenuSection.cpp b/Engine/Source/Developer/ToolMenus/Private/ToolMenuSection.cpp
--- a/Engine/Source/Developer/ToolMenus/Private/ToolMenuSection.cpp
+++ b/Engine/Source/Developer/ToolMenus/Private/ToolMenuSection.cpp
@@ -12,7 +12,8 @@
 
 FToolMenuSection::FToolMenuSection() :
 	ToolMenuSectionDynamic(nullptr),
-	bIsRegistering(false)
+	bIsRegistering(false),
+	bAddedDuringRegister(false)
 {
 
 }
diff --git a/Engine/Source/Developer/ToolMenus/Public/ToolMenuSection.h b/Engine/Source/Developer/ToolMenus/Public/ToolMenuSection.h
--- a/Engine/Source/Developer/ToolMenus/Public/ToolMenuSection.h
+++ b/Engine/Source/Developer/ToolMenus/Public/ToolMenuSection.h
@@ -97,4 +97,5 @@ public:
 private:
 
 	bool bIsRegistering;
+	bool bAddedDuringRegister;
 };
