Tool Menus: Script API to add section now skips overriding label or position if not set
#rnx

diff --git a/Engine/Source/Developer/ToolMenus/Private/ToolMenu.cpp b/Engine/Source/Developer/ToolMenus/Private/ToolMenu.cpp
--- a/Engine/Source/Developer/ToolMenus/Private/ToolMenu.cpp
+++ b/Engine/Source/Developer/ToolMenus/Private/ToolMenu.cpp
@@ -140,7 +140,7 @@ FToolMenuSection& UToolMenu::AddSection(const FName SectionName, const TAttribut
 			Sections[InsertIndex].Label = InLabel;
 		}
 
-		if (InPosition.Name != NAME_None)
+		if (InPosition.Name != NAME_None || InPosition.Position != EToolMenuInsertType::Default)
 		{
 			Sections[InsertIndex].InsertPosition = InPosition;
 		}
@@ -192,8 +192,16 @@ FToolMenuSection& UToolMenu::AddSection(const FName SectionName, const TAttribut
 void UToolMenu::AddSectionScript(const FName SectionName, const FText& InLabel, const FName InsertName, const EToolMenuInsertType InsertType)
 {
 	FToolMenuSection& Section = FindOrAddSection(SectionName);
-	Section.Label = InLabel;
-	Section.InsertPosition = FToolMenuInsert(InsertName, InsertType);
+
+	if (!InLabel.IsEmpty())
+	{
+		Section.Label = InLabel;
+	}
+
+	if (InsertName != NAME_None || InsertType != EToolMenuInsertType::Default)
+	{
+		Section.InsertPosition = FToolMenuInsert(InsertName, InsertType);
+	}
 }
 
 void UToolMenu::AddDynamicSectionScript(const FName SectionName, UToolMenuSectionDynamic* InObject)
