Force re-copy files from Main that failed to merge correctly
#rnx

diff --git a/Engine/Source/Developer/ToolMenus/Public/ToolMenu.h b/Engine/Source/Developer/ToolMenus/Public/ToolMenu.h
--- a/Engine/Source/Developer/ToolMenus/Public/ToolMenu.h
+++ b/Engine/Source/Developer/ToolMenus/Public/ToolMenu.h
@@ -62,6 +62,9 @@ public:
 	/** returns "SubMenuC.SubMenuD" for menu "ModuleA.MenuB.SubMenuC.SubMenuD" */
 	FString GetSubMenuNamePath() const;
 
+	/* Set support for extenders */
+	void SetExtendersEnabled(bool bEnabled);
+
 	//~ Begin UToolMenuBase Interface
 	virtual bool IsEditing() const override;
 	virtual FName GetSectionName(const FName InEntryName) const override;
@@ -82,7 +85,7 @@ public:
 	template <typename TContextType>
 	TContextType* FindContext() const
 	{
-		return Context.Find<TContextType>();
+		return Context.FindContext<TContextType>();
 	}
 
 	//~ Begin UObject Interface
@@ -105,6 +108,8 @@ private:
 
 	int32 FindInsertIndex(const FToolMenuSection& InSection) const;
 
+	bool IsRegistering() const;
+
 public:
 
 	UPROPERTY(EditAnywhere, BlueprintReadWrite, Category = "Tool Menus")
@@ -161,6 +166,8 @@ public:
 private:
 
 	bool bRegistered;
+	bool bIsRegistering;
+	bool bExtendersEnabled;
 
 	const ISlateStyle* StyleSet;
 
diff --git a/Engine/Source/Developer/ToolMenus/Public/ToolMenus.h b/Engine/Source/Developer/ToolMenus/Public/ToolMenus.h
--- a/Engine/Source/Developer/ToolMenus/Public/ToolMenus.h
+++ b/Engine/Source/Developer/ToolMenus/Public/ToolMenus.h
@@ -327,7 +327,7 @@ private:
 	void PopulateMenuBarBuilder(FMenuBarBuilder& MenuBarBuilder, UToolMenu* MenuData);
 	void PopulateToolBarBuilder(FToolBarBuilder& ToolBarBuilder, UToolMenu* MenuData);
 
-	TSharedRef<SWidget> GenerateToolbarComboButtonMenu(const FName SubMenuFullName, FToolMenuContext InContext);
+	TSharedRef<SWidget> GenerateToolbarComboButtonMenu(TWeakObjectPtr<UToolMenu> InParent, const FName InBlockName);
 
 	FOnGetContent ConvertWidgetChoice(const FNewToolMenuWidgetChoice& Choice, const FToolMenuContext& Context) const;
 
diff --git a/Engine/Source/Runtime/Slate/Public/Framework/MultiBox/MultiBoxBuilder.h b/Engine/Source/Runtime/Slate/Public/Framework/MultiBox/MultiBoxBuilder.h
--- a/Engine/Source/Runtime/Slate/Public/Framework/MultiBox/MultiBoxBuilder.h
+++ b/Engine/Source/Runtime/Slate/Public/Framework/MultiBox/MultiBoxBuilder.h
@@ -119,6 +119,12 @@ public:
 	/** @return  The customization settings for the box being built */
 	FMultiBoxCustomization GetCustomization() const;
 
+	/** Sets extender support */
+	void SetExtendersEnabled(bool bEnabled) { bExtendersEnabled = bEnabled; }
+
+	/** @return True if extenders are enabled */
+	bool ExtendersEnabled() const { return bExtendersEnabled; }
+
 protected:
 	/** Applies any potential extension hooks at the current place */
 	virtual void ApplyHook(FName InExtensionHook, EExtensionHook::Position HookPosition) = 0;
@@ -143,6 +149,9 @@ protected:
 
 	/** Name of the menu */
 	FName MenuName;
+
+	/** If extenders are enabled */
+	bool bExtendersEnabled;
 };
 
 
