DR - Missed file
#rnx

diff --git a/Engine/Source/Developer/ToolMenus/Private/ToolMenus.cpp b/Engine/Source/Developer/ToolMenus/Private/ToolMenus.cpp
--- a/Engine/Source/Developer/ToolMenus/Private/ToolMenus.cpp
+++ b/Engine/Source/Developer/ToolMenus/Private/ToolMenus.cpp
@@ -597,23 +597,6 @@ void UToolMenus::ApplyCustomization(UToolMenu* GeneratedMenu)
 		}
 	}
 
-	// Hide items based on whitelist
-	if (CustomizedMenu.WhitelistEntries.Num() > 0)
-	{
-		for (int32 SectionIndex = 0; SectionIndex < NewSections.Num(); ++SectionIndex)
-		{
-			FToolMenuSection& Section = NewSections[SectionIndex];
-			for (int32 i = 0; i < Section.Blocks.Num(); ++i)
-			{
-				if (!CustomizedMenu.IsEntryWhitelisted(Section.Blocks[i].Name))
-				{
-					Section.Blocks.RemoveAt(i);
-					--i;
-				}
-			}
-		}
-	}
-
 	// Hide sections and entries
 	if (!GeneratedMenu->IsEditing())
 	{
