Refresh widgets of tool menus that are extended automatically

diff --git a/Engine/Source/Developer/ToolMenus/Private/ToolMenus.cpp b/Engine/Source/Developer/ToolMenus/Private/ToolMenus.cpp
--- a/Engine/Source/Developer/ToolMenus/Private/ToolMenus.cpp
+++ b/Engine/Source/Developer/ToolMenus/Private/ToolMenus.cpp
@@ -1935,6 +1935,10 @@ UToolMenu* UToolMenus::ExtendMenu(const FName InName)
 		{
 			Section.bIsRegistering = Found->bIsRegistering;
 		}
+
+		// Refresh all widgets because this could be child of another menu being displayed
+		RefreshAllWidgets();
+
 		return Found;
 	}
 
