Fix tool menus not supporting styles

diff --git a/Engine/Source/Developer/ToolMenus/Private/ToolMenus.cpp b/Engine/Source/Developer/ToolMenus/Private/ToolMenus.cpp
--- a/Engine/Source/Developer/ToolMenus/Private/ToolMenus.cpp
+++ b/Engine/Source/Developer/ToolMenus/Private/ToolMenus.cpp
@@ -1658,6 +1658,7 @@ TSharedRef<SWidget> UToolMenus::GenerateWidget(UToolMenu* GeneratedMenu)
 	// Copy native properties that serialize does not
 	GeneratedMenuWidget.GeneratedMenu->Context = GeneratedMenu->Context;
 	GeneratedMenuWidget.GeneratedMenu->StyleSet = GeneratedMenu->StyleSet;
+	GeneratedMenuWidget.GeneratedMenu->StyleName = GeneratedMenu->StyleName;
 
 	if (GeneratedMenu->IsEditing())
 	{
@@ -1674,6 +1675,12 @@ TSharedRef<SWidget> UToolMenus::GenerateWidget(UToolMenu* GeneratedMenu)
 		}
 
 		FMenuBuilder MenuBuilder(GeneratedMenu->bShouldCloseWindowAfterMenuSelection, GeneratedMenu->Context.CommandList, GeneratedMenu->Context.GetAllExtenders(), GeneratedMenu->bCloseSelfOnly, GeneratedMenu->StyleSet, GeneratedMenu->bSearchable, GeneratedMenu->MenuName);
+
+		if (GeneratedMenu->StyleName != NAME_None)
+		{
+			MenuBuilder.SetStyle(GeneratedMenu->StyleSet, GeneratedMenu->StyleName);
+		}
+
 		MenuBuilder.SetExtendersEnabled(GeneratedMenu->bExtendersEnabled);
 		MenuBuilder.GetMultiBox()->WeakToolMenu = GeneratedMenu;
 		PopulateMenuBuilder(MenuBuilder, GeneratedMenu);
@@ -1688,6 +1695,12 @@ TSharedRef<SWidget> UToolMenus::GenerateWidget(UToolMenu* GeneratedMenu)
 	else if (GeneratedMenu->MenuType == EMultiBoxType::Menu)
 	{
 		FMenuBuilder MenuBuilder(GeneratedMenu->bShouldCloseWindowAfterMenuSelection, GeneratedMenu->Context.CommandList, GeneratedMenu->Context.GetAllExtenders(), GeneratedMenu->bCloseSelfOnly, GeneratedMenu->StyleSet, GeneratedMenu->bSearchable, GeneratedMenu->MenuName);
+
+		if (GeneratedMenu->StyleName != NAME_None)
+		{
+			MenuBuilder.SetStyle(GeneratedMenu->StyleSet, GeneratedMenu->StyleName);
+		}
+
 		MenuBuilder.SetExtendersEnabled(GeneratedMenu->bExtendersEnabled);
 		PopulateMenuBuilder(MenuBuilder, GeneratedMenu);
 		TSharedRef<SWidget> Result = MenuBuilder.MakeWidget(nullptr, GeneratedMenu->MaxHeight);
@@ -1697,6 +1710,12 @@ TSharedRef<SWidget> UToolMenus::GenerateWidget(UToolMenu* GeneratedMenu)
 	else if (GeneratedMenu->MenuType == EMultiBoxType::MenuBar)
 	{
 		FMenuBarBuilder MenuBarBuilder(GeneratedMenu->Context.CommandList, GeneratedMenu->Context.GetAllExtenders(), GeneratedMenu->StyleSet, GeneratedMenu->MenuName);
+
+		if (GeneratedMenu->StyleName != NAME_None)
+		{
+			MenuBarBuilder.SetStyle(GeneratedMenu->StyleSet, GeneratedMenu->StyleName);
+		}
+
 		MenuBarBuilder.SetExtendersEnabled(GeneratedMenu->bExtendersEnabled);
 		PopulateMenuBarBuilder(MenuBarBuilder, GeneratedMenu);
 		TSharedRef<SWidget> Result = MenuBarBuilder.MakeWidget();
@@ -1709,6 +1728,12 @@ TSharedRef<SWidget> UToolMenus::GenerateWidget(UToolMenu* GeneratedMenu)
 		FToolBarBuilder ToolbarBuilder(GeneratedMenu->Context.CommandList, GeneratedMenu->MenuName, GeneratedMenu->Context.GetAllExtenders(), ToolBarOrientation, GeneratedMenu->bToolBarForceSmallIcons);
 		ToolbarBuilder.SetExtendersEnabled(GeneratedMenu->bExtendersEnabled);
 		ToolbarBuilder.SetIsFocusable(GeneratedMenu->bToolBarIsFocusable);
+
+		if (GeneratedMenu->StyleName != NAME_None)
+		{
+			ToolbarBuilder.SetStyle(GeneratedMenu->StyleSet, GeneratedMenu->StyleName);
+		}
+
 		PopulateToolBarBuilder(ToolbarBuilder, GeneratedMenu);
 		TSharedRef<SWidget> Result = ToolbarBuilder.MakeWidget();
 		GeneratedMenuWidget.Widget = Result;
@@ -1853,6 +1878,12 @@ bool UToolMenus::RefreshMenuWidget(const FName InName, FGeneratedToolMenuWidget&
 	{
 		FMenuBuilder MenuBuilder(GeneratedMenu->bShouldCloseWindowAfterMenuSelection, GeneratedMenu->Context.CommandList, GeneratedMenu->Context.GetAllExtenders(), GeneratedMenu->bCloseSelfOnly, GeneratedMenu->StyleSet, GeneratedMenu->bSearchable);
 		MenuBuilder.SetExtendersEnabled(GeneratedMenu->bExtendersEnabled);
+
+		if (GeneratedMenu->StyleName != NAME_None)
+		{
+			MenuBuilder.SetStyle(GeneratedMenu->StyleSet, GeneratedMenu->StyleName);
+		}
+
 		PopulateMenuBuilder(MenuBuilder, GeneratedMenu);
 		MultiBoxWidget->SetMultiBox(MenuBuilder.GetMultiBox());
 	}
@@ -1860,6 +1891,12 @@ bool UToolMenus::RefreshMenuWidget(const FName InName, FGeneratedToolMenuWidget&
 	{
 		FMenuBarBuilder MenuBarBuilder(GeneratedMenu->Context.CommandList, GeneratedMenu->Context.GetAllExtenders(), GeneratedMenu->StyleSet);
 		MenuBarBuilder.SetExtendersEnabled(GeneratedMenu->bExtendersEnabled);
+
+		if (GeneratedMenu->StyleName != NAME_None)
+		{
+			MenuBarBuilder.SetStyle(GeneratedMenu->StyleSet, GeneratedMenu->StyleName);
+		}
+
 		PopulateMenuBarBuilder(MenuBarBuilder, GeneratedMenu);
 		MultiBoxWidget->SetMultiBox(MenuBarBuilder.GetMultiBox());
 	}
@@ -1869,6 +1906,12 @@ bool UToolMenus::RefreshMenuWidget(const FName InName, FGeneratedToolMenuWidget&
 		FToolBarBuilder ToolbarBuilder(GeneratedMenu->Context.CommandList, GeneratedMenu->MenuName, GeneratedMenu->Context.GetAllExtenders(), ToolBarOrientation, GeneratedMenu->bToolBarForceSmallIcons);
 		ToolbarBuilder.SetExtendersEnabled(GeneratedMenu->bExtendersEnabled);
 		ToolbarBuilder.SetIsFocusable(GeneratedMenu->bToolBarIsFocusable);
+
+		if (GeneratedMenu->StyleName != NAME_None)
+		{
+			ToolbarBuilder.SetStyle(GeneratedMenu->StyleSet, GeneratedMenu->StyleName);
+		}
+
 		PopulateToolBarBuilder(ToolbarBuilder, GeneratedMenu);
 		MultiBoxWidget->SetMultiBox(ToolbarBuilder.GetMultiBox());
 	}
