Fixed not being able to add separators to toolbars using the new tool menus api to create toolbars.  Toolbars were not accounting for multiple separator types.  Deprecated the concept of a menu separator and a toolbar separator and made only one separator type.
For now did not deprecate usage of AddMenuSeparator in menu builders due to the extreme number of fiiles this touches

diff --git a/Engine/Source/Developer/ToolMenus/Private/ToolMenuEntry.cpp b/Engine/Source/Developer/ToolMenus/Private/ToolMenuEntry.cpp
--- a/Engine/Source/Developer/ToolMenus/Private/ToolMenuEntry.cpp
+++ b/Engine/Source/Developer/ToolMenus/Private/ToolMenuEntry.cpp
@@ -187,15 +187,11 @@ FToolMenuEntry FToolMenuEntry::InitComboButton(const FName InName, const FToolUI
 	return Entry;
 }
 
-FToolMenuEntry FToolMenuEntry::InitMenuSeparator(const FName InName)
+FToolMenuEntry FToolMenuEntry::InitSeparator(const FName InName)
 {
-	return FToolMenuEntry(UToolMenus::Get()->CurrentOwner(), InName, EMultiBlockType::MenuSeparator);
+	return FToolMenuEntry(UToolMenus::Get()->CurrentOwner(), InName, EMultiBlockType::Separator);
 }
 
-FToolMenuEntry FToolMenuEntry::InitToolBarSeparator(const FName InName)
-{
-	return FToolMenuEntry(UToolMenus::Get()->CurrentOwner(), InName, EMultiBlockType::ToolBarSeparator);
-}
 
 FToolMenuEntry FToolMenuEntry::InitWidget(const FName InName, const TSharedRef<SWidget>& InWidget, const FText& Label, bool bNoIndent, bool bSearchable)
 {
diff --git a/Engine/Source/Developer/ToolMenus/Private/ToolMenuSection.cpp b/Engine/Source/Developer/ToolMenus/Private/ToolMenuSection.cpp
--- a/Engine/Source/Developer/ToolMenus/Private/ToolMenuSection.cpp
+++ b/Engine/Source/Developer/ToolMenus/Private/ToolMenuSection.cpp
@@ -117,7 +117,12 @@ FToolMenuEntry& FToolMenuSection::AddDynamicEntry(const FName InName, const FNew
 
 FToolMenuEntry& FToolMenuSection::AddMenuSeparator(const FName InName)
 {
-	return AddEntry(FToolMenuEntry::InitMenuSeparator(InName));
+	return AddSeparator(InName);
+}
+
+FToolMenuEntry& FToolMenuSection::AddSeparator(const FName InName)
+{
+	return AddEntry(FToolMenuEntry::InitSeparator(InName));
 }
 
 FToolMenuEntry& FToolMenuSection::AddSubMenu(const FName InName, const TAttribute<FText>& InLabel, const TAttribute<FText>& InToolTip, const FNewToolMenuChoice& InMakeMenu, const FToolUIActionChoice& InAction, const EUserInterfaceActionType InUserInterfaceActionType, bool bInOpenSubMenuOnClick, const TAttribute<FSlateIcon>& InIcon, const bool bInShouldCloseWindowAfterMenuSelection)
diff --git a/Engine/Source/Developer/ToolMenus/Private/ToolMenus.cpp b/Engine/Source/Developer/ToolMenus/Private/ToolMenus.cpp
--- a/Engine/Source/Developer/ToolMenus/Private/ToolMenus.cpp
+++ b/Engine/Source/Developer/ToolMenus/Private/ToolMenus.cpp
@@ -1061,7 +1061,7 @@ void UToolMenus::PopulateMenuBuilder(FMenuBuilder& MenuBuilder, UToolMenu* MenuD
 					}
 				}
 			}
-			else if (Block.Type == EMultiBlockType::MenuSeparator)
+			else if (Block.Type == EMultiBlockType::Separator)
 			{
 				MenuBuilder.AddMenuSeparator(Block.Name);
 			}
@@ -1186,7 +1186,7 @@ void UToolMenus::PopulateToolBarBuilder(FToolBarBuilder& ToolBarBuilder, UToolMe
 					ToolBarBuilder.AddComboButton(UIAction, Delegate, Block.Label, Block.ToolTip, Block.Icon, Block.ToolBarData.bSimpleComboBox, Block.TutorialHighlightName);
 				}
 			}
-			else if (Block.Type == EMultiBlockType::ToolBarSeparator)
+			else if (Block.Type == EMultiBlockType::Separator)
 			{
 				ToolBarBuilder.AddSeparator(Block.Name);
 			}
@@ -1757,10 +1757,6 @@ void UToolMenus::ModifyEntryForEditDialog(FToolMenuEntry& Entry)
 			Entry.SubMenuData.bIsSubMenu = true;
 		}
 	}
-	else if (Entry.Type == EMultiBlockType::ToolBarSeparator)
-	{
-		Entry.Type = EMultiBlockType::MenuSeparator;
-	}
 }
 
 void UToolMenus::AssignSetTimerForNextTickDelegate(const FSimpleDelegate& InDelegate)
diff --git a/Engine/Source/Developer/ToolMenus/Public/ToolMenuEntry.h b/Engine/Source/Developer/ToolMenus/Public/ToolMenuEntry.h
--- a/Engine/Source/Developer/ToolMenus/Public/ToolMenuEntry.h
+++ b/Engine/Source/Developer/ToolMenus/Public/ToolMenuEntry.h
@@ -98,8 +98,7 @@ struct TOOLMENUS_API FToolMenuEntry
 	static FToolMenuEntry InitToolBarButton(const TSharedPtr< const FUICommandInfo >& InCommand, const TAttribute<FText>& InLabelOverride = TAttribute<FText>(), const TAttribute<FText>& InToolTipOverride = TAttribute<FText>(), const TAttribute<FSlateIcon>& InIconOverride = TAttribute<FSlateIcon>(), FName InTutorialHighlightName = NAME_None, const TOptional<FName> InNameOverride = TOptional<FName>());
 	static FToolMenuEntry InitComboButton(const FName InName, const FToolUIActionChoice& InAction, const FNewToolMenuChoice& InMenuContentGenerator, const TAttribute<FText>& InLabelOverride = TAttribute<FText>(), const TAttribute<FText>& InToolTipOverride = TAttribute<FText>(), const TAttribute<FSlateIcon>& InIconOverride = TAttribute<FSlateIcon>(), bool bInSimpleComboBox = false, FName InTutorialHighlightName = NAME_None);
 
-	static FToolMenuEntry InitMenuSeparator(const FName InName);
-	static FToolMenuEntry InitToolBarSeparator(const FName InName);
+	static FToolMenuEntry InitSeparator(const FName InName);
 
 	static FToolMenuEntry InitWidget(const FName InName, const TSharedRef<SWidget>& InWidget, const FText& Label, bool bNoIndent = false, bool bSearchable = true);
 
diff --git a/Engine/Source/Developer/ToolMenus/Public/ToolMenuSection.h b/Engine/Source/Developer/ToolMenus/Public/ToolMenuSection.h
--- a/Engine/Source/Developer/ToolMenus/Public/ToolMenuSection.h
+++ b/Engine/Source/Developer/ToolMenus/Public/ToolMenuSection.h
@@ -41,7 +41,10 @@ public:
 
 	FToolMenuEntry& AddDynamicEntry(const FName InName, const FNewToolMenuSectionDelegate& InConstruct);
 	FToolMenuEntry& AddDynamicEntry(const FName InName, const FNewToolMenuDelegateLegacy& InConstruct);
+
+	UE_DEPRECATED(4.26, "AddMenuSeparator has been deprecated.  Use AddSeparator instead.")
 	FToolMenuEntry& AddMenuSeparator(const FName InName);
+	FToolMenuEntry& AddSeparator(const FName InName);
 
 	FToolMenuEntry& AddSubMenu(const FName InName, const TAttribute<FText>& InLabel, const TAttribute<FText>& InToolTip, const FNewToolMenuChoice& InMakeMenu, const FToolUIActionChoice& InAction, const EUserInterfaceActionType InUserInterfaceActionType, bool bInOpenSubMenuOnClick = false, const TAttribute<FSlateIcon>& InIcon = TAttribute<FSlateIcon>(), const bool bShouldCloseWindowAfterMenuSelection = true);
 	FToolMenuEntry& AddSubMenu(const FName InName, const TAttribute<FText>& InLabel, const TAttribute<FText>& InToolTip, const FNewToolMenuChoice& InMakeMenu, bool bInOpenSubMenuOnClick = false, const TAttribute<FSlateIcon>& InIcon = TAttribute<FSlateIcon>(), const bool bShouldCloseWindowAfterMenuSelection = true);
diff --git a/Engine/Source/Editor/CommonMenuExtensions/Private/ShowFlagMenuCommands.cpp b/Engine/Source/Editor/CommonMenuExtensions/Private/ShowFlagMenuCommands.cpp
--- a/Engine/Source/Editor/CommonMenuExtensions/Private/ShowFlagMenuCommands.cpp
+++ b/Engine/Source/Editor/CommonMenuExtensions/Private/ShowFlagMenuCommands.cpp
@@ -119,7 +119,7 @@ void FShowFlagMenuCommands::CreateShowFlagsSubMenu(UToolMenu* Menu, TArray<uint3
 
 		if (ArrayIndex == EntryOffset - 1)
 		{
-			Section.AddMenuSeparator(NAME_None);
+			Section.AddSeparator(NAME_None);
 		}
 	}
 }
diff --git a/Engine/Source/Editor/ContentBrowser/Private/NewAssetOrClassContextMenu.cpp b/Engine/Source/Editor/ContentBrowser/Private/NewAssetOrClassContextMenu.cpp
--- a/Engine/Source/Editor/ContentBrowser/Private/NewAssetOrClassContextMenu.cpp
+++ b/Engine/Source/Editor/ContentBrowser/Private/NewAssetOrClassContextMenu.cpp
@@ -460,7 +460,7 @@ void FNewAssetOrClassContextMenu::CreateNewAssetMenus(UToolMenu* Menu, FName Sec
 		return;
 	}
 
-	Section.AddMenuSeparator(NAME_None);
+	Section.AddSeparator(NAME_None);
 
 	TArray<TSharedPtr<FCategorySubMenuItem>> SortedMenus;
 	SubMenuData->Children.GenerateValueArray(SortedMenus);
diff --git a/Engine/Source/Editor/LevelEditor/Private/LevelEditorMenu.cpp b/Engine/Source/Editor/LevelEditor/Private/LevelEditorMenu.cpp
--- a/Engine/Source/Editor/LevelEditor/Private/LevelEditorMenu.cpp
+++ b/Engine/Source/Editor/LevelEditor/Private/LevelEditorMenu.cpp
@@ -98,7 +98,7 @@ void FLevelEditorMenu::RegisterLevelEditorMenus()
 								ToggleFavoriteLabel.BindStatic(&Local::GetToggleFavoriteLabelText);
 								Section.AddMenuEntry(FLevelEditorCommands::Get().ToggleFavorite, ToggleFavoriteLabel);
 							}
-							Section.AddMenuSeparator("LevelEditorToggleFavorite");
+							Section.AddSeparator("LevelEditorToggleFavorite");
 						}
 						const FMainMRUFavoritesList& MRUFavorites = *FModuleManager::LoadModuleChecked<IMainFrameModule>("MainFrame").GetMRUFavoritesList();
 						const int32 NumFavorites = MRUFavorites.GetNumFavorites();
@@ -209,7 +209,7 @@ void FLevelEditorMenu::RegisterLevelEditorMenus()
 
 				Section.AddMenuEntry( FLevelEditorCommands::Get().BrowseCVars );
 
-				Section.AddMenuSeparator( "HelpBrowse" );
+				Section.AddSeparator( "HelpBrowse" );
 
 				Section.AddMenuEntry( FLevelEditorCommands::Get().BrowseViewportControls );
 			}
diff --git a/Engine/Source/Editor/LevelEditor/Private/LevelEditorToolBar.cpp b/Engine/Source/Editor/LevelEditor/Private/LevelEditorToolBar.cpp
--- a/Engine/Source/Editor/LevelEditor/Private/LevelEditorToolBar.cpp
+++ b/Engine/Source/Editor/LevelEditor/Private/LevelEditorToolBar.cpp
@@ -2183,7 +2183,7 @@ void FLevelEditorToolBar::RegisterSourceControlMenu()
 		}
 	}));
 
-	Section.AddMenuSeparator("SourceControlConnectionSeparator");
+	Section.AddSeparator("SourceControlConnectionSeparator");
 
 	Section.AddMenuEntry(
 		FLevelEditorCommands::Get().CheckOutModifiedFiles,
@@ -2521,7 +2521,7 @@ void FLevelEditorToolBar::RegisterEditorModesMenu()
 		// Build Default Modes first
 		BuildEditorModes(MakeArrayView(DefaultModes));
 
-		InSection.AddMenuSeparator(NAME_None);
+		InSection.AddSeparator(NAME_None);
 			
 		// Build non-default modes second
 		BuildEditorModes(MakeArrayView(NonDefaultModes));
diff --git a/Engine/Source/Editor/LevelEditor/Private/SLevelViewportToolBar.cpp b/Engine/Source/Editor/LevelEditor/Private/SLevelViewportToolBar.cpp
--- a/Engine/Source/Editor/LevelEditor/Private/SLevelViewportToolBar.cpp
+++ b/Engine/Source/Editor/LevelEditor/Private/SLevelViewportToolBar.cpp
@@ -92,7 +92,7 @@ static void FillShowMenuStatic(UToolMenu* Menu, TArray< FLevelViewportCommands::
 
 		if (EntryIndex == EntryOffset - 1)
 		{
-			Section.AddMenuSeparator(NAME_None);
+			Section.AddSeparator(NAME_None);
 		}
 	}
 }
@@ -664,7 +664,7 @@ void SLevelViewportToolBar::FillOptionsMenu(UToolMenu* Menu)
 					FSlateIcon(),
 					Action);
 
-				Section.AddMenuSeparator("DisableRealtimeOverrideSeparator");
+				Section.AddSeparator("DisableRealtimeOverrideSeparator");
 					
 			}
 
@@ -1522,7 +1522,7 @@ void SLevelViewportToolBar::FillViewMenu(UToolMenu* Menu)
 						EUserInterfaceActionType::RadioButton
 					);
 
-					SubMenuSection.AddMenuSeparator("LandscapeLODSeparator");
+					SubMenuSection.AddSeparator("LandscapeLODSeparator");
 
 					static const FText FormatString = LOCTEXT("LandscapeLODFixed", "Fixed at {0}");
 					for (int32 i = 0; i < 8; ++i)
diff --git a/Engine/Source/Editor/MainFrame/Private/Menus/LayoutsMenu.cpp b/Engine/Source/Editor/MainFrame/Private/Menus/LayoutsMenu.cpp
--- a/Engine/Source/Editor/MainFrame/Private/Menus/LayoutsMenu.cpp
+++ b/Engine/Source/Editor/MainFrame/Private/Menus/LayoutsMenu.cpp
@@ -792,7 +792,7 @@ void FLayoutsMenuLoad::MakeLoadLayoutsMenu(UToolMenu* InToolMenu)
 		// Separator
 		if (FLayoutsMenuBase::IsThereUserLayouts())
 		{
-			Section.AddMenuSeparator("AdditionalSectionsSeparator");
+			Section.AddSeparator("AdditionalSectionsSeparator");
 		}
 
 		// Import...
@@ -980,7 +980,7 @@ void FLayoutsMenuSave::MakeSaveLayoutsMenu(UToolMenu* InToolMenu)
 			// Separator
 			if (FLayoutsMenuBase::IsThereUserLayouts())
 			{
-				Section.AddMenuSeparator("AdditionalSectionsSeparator");
+				Section.AddSeparator("AdditionalSectionsSeparator");
 			}
 
 			// Save as...
@@ -1090,7 +1090,7 @@ void FLayoutsMenuRemove::MakeRemoveLayoutsMenu(UToolMenu* InToolMenu)
 			// Separator
 			if (FLayoutsMenuBase::IsThereUserLayouts())
 			{
-				Section.AddMenuSeparator("AdditionalSectionsSeparator");
+				Section.AddSeparator("AdditionalSectionsSeparator");
 			}
 
 			// Remove all
diff --git a/Engine/Source/Editor/MainFrame/Private/Menus/MainMenu.cpp b/Engine/Source/Editor/MainFrame/Private/Menus/MainMenu.cpp
--- a/Engine/Source/Editor/MainFrame/Private/Menus/MainMenu.cpp
+++ b/Engine/Source/Editor/MainFrame/Private/Menus/MainMenu.cpp
@@ -286,7 +286,7 @@ void FMainMenu::RegisterWindowMenu()
 		// Enable Fullscreen section
 #if !PLATFORM_MAC && !PLATFORM_LINUX // On Mac/Linux windowed fullscreen mode in the editor is currently unavailable
 		// Separator
-		Section.AddMenuSeparator("FullscreenSeparator");
+		Section.AddSeparator("FullscreenSeparator");
 		// Fullscreen
 		Section.AddMenuEntry(FMainFrameCommands::Get().ToggleFullscreen);
 #endif
@@ -317,10 +317,10 @@ void FMainMenu::RegisterHelpMenu()
 
 		const FText SupportWebSiteLabel = NSLOCTEXT("MainHelpMenu", "VisitUnrealEngineSupportWebSite", "Unreal Engine Support Web Site...");
 
-		HelpOnlineSection.AddMenuSeparator("EpicGamesHelp");
+		HelpOnlineSection.AddSeparator("EpicGamesHelp");
 		HelpOnlineSection.AddMenuEntry(FMainFrameCommands::Get().VisitEpicGamesDotCom);
 
-		HelpOnlineSection.AddMenuSeparator("Credits");
+		HelpOnlineSection.AddSeparator("Credits");
 		HelpOnlineSection.AddMenuEntry(FMainFrameCommands::Get().CreditsUnrealEd);
 	}
 
diff --git a/Engine/Source/Editor/UnrealEd/Private/Toolkits/AssetEditorToolkit.cpp b/Engine/Source/Editor/UnrealEd/Private/Toolkits/AssetEditorToolkit.cpp
--- a/Engine/Source/Editor/UnrealEd/Private/Toolkits/AssetEditorToolkit.cpp
+++ b/Engine/Source/Editor/UnrealEd/Private/Toolkits/AssetEditorToolkit.cpp
@@ -928,7 +928,7 @@ void FAssetEditorToolkit::FillDefaultFileMenuCommands(FToolMenuSection& InSectio
 	{
 		InSection.AddMenuEntry(FAssetEditorCommonCommands::Get().SaveAssetAs, TAttribute<FText>(), TAttribute<FText>(), FSlateIcon(FEditorStyle::GetStyleSetName(), "AssetEditor.SaveAssetAs.Small")).InsertPosition = InsertPosition;
 	}
-	InSection.AddMenuSeparator("DefaultFileMenuCommandsSeparator").InsertPosition = InsertPosition;;
+	InSection.AddSeparator("DefaultFileMenuCommandsSeparator").InsertPosition = InsertPosition;;
 
 	if( IsWorldCentricAssetEditor() )
 	{
diff --git a/Engine/Source/Runtime/Slate/Private/Framework/MultiBox/MultiBox.cpp b/Engine/Source/Runtime/Slate/Private/Framework/MultiBox/MultiBox.cpp
--- a/Engine/Source/Runtime/Slate/Private/Framework/MultiBox/MultiBox.cpp
+++ b/Engine/Source/Runtime/Slate/Private/Framework/MultiBox/MultiBox.cpp
@@ -822,7 +822,7 @@ void SMultiBoxWidget::BuildMultiBoxWidget()
 					bSectionContainsIcons = true;
 				}
 
-				if (TestBlock.GetType() == EMultiBlockType::MenuSeparator)
+				if (TestBlock.GetType() == EMultiBlockType::Separator)
 				{
 					break;
 				}
diff --git a/Engine/Source/Runtime/Slate/Private/Framework/MultiBox/MultiBoxBuilder.cpp b/Engine/Source/Runtime/Slate/Private/Framework/MultiBox/MultiBoxBuilder.cpp
--- a/Engine/Source/Runtime/Slate/Private/Framework/MultiBox/MultiBoxBuilder.cpp
+++ b/Engine/Source/Runtime/Slate/Private/Framework/MultiBox/MultiBoxBuilder.cpp
@@ -215,16 +215,21 @@ void FMenuBuilder::EndSection()
 }
 
 void FMenuBuilder::AddMenuSeparator(FName InExtensionHook)
+{
+	AddSeparator(InExtensionHook);
+}
+
+void FMenuBuilder::AddSeparator(FName InExtensionHook /*= NAME_None*/)
 {
 	ApplySectionBeginning();
 
 	ApplyHook(InExtensionHook, EExtensionHook::Before);
 
 	// Never add a menu separate as the first item, even if we were asked to
-	if( MultiBox->GetBlocks().Num() > 0 || FMultiBoxSettings::DisplayMultiboxHooks.Get() )
+	if (MultiBox->GetBlocks().Num() > 0 || FMultiBoxSettings::DisplayMultiboxHooks.Get())
 	{
-		TSharedRef< FMenuSeparatorBlock > NewMenuSeparatorBlock( new FMenuSeparatorBlock(InExtensionHook, /* bInIsPartOfHeading=*/ false) );
-		MultiBox->AddMultiBlock( NewMenuSeparatorBlock );
+		TSharedRef< FMenuSeparatorBlock > NewMenuSeparatorBlock(new FMenuSeparatorBlock(InExtensionHook, /* bInIsPartOfHeading=*/ false));
+		MultiBox->AddMultiBlock(NewMenuSeparatorBlock);
 	}
 
 	ApplyHook(InExtensionHook, EExtensionHook::After);
diff --git a/Engine/Source/Runtime/Slate/Private/Framework/MultiBox/SMenuSeparatorBlock.cpp b/Engine/Source/Runtime/Slate/Private/Framework/MultiBox/SMenuSeparatorBlock.cpp
--- a/Engine/Source/Runtime/Slate/Private/Framework/MultiBox/SMenuSeparatorBlock.cpp
+++ b/Engine/Source/Runtime/Slate/Private/Framework/MultiBox/SMenuSeparatorBlock.cpp
@@ -9,7 +9,7 @@
  * Constructor
  */
 FMenuSeparatorBlock::FMenuSeparatorBlock(const FName& InExtensionHook, bool bInIsPartOfHeading)
-	: FMultiBlock( nullptr, nullptr, InExtensionHook, EMultiBlockType::MenuSeparator, bInIsPartOfHeading )
+	: FMultiBlock( nullptr, nullptr, InExtensionHook, EMultiBlockType::Separator, bInIsPartOfHeading )
 {
 }
 
diff --git a/Engine/Source/Runtime/Slate/Private/Framework/MultiBox/SToolBarSeparatorBlock.cpp b/Engine/Source/Runtime/Slate/Private/Framework/MultiBox/SToolBarSeparatorBlock.cpp
--- a/Engine/Source/Runtime/Slate/Private/Framework/MultiBox/SToolBarSeparatorBlock.cpp
+++ b/Engine/Source/Runtime/Slate/Private/Framework/MultiBox/SToolBarSeparatorBlock.cpp
@@ -10,7 +10,7 @@
  * Constructor
  */
 FToolBarSeparatorBlock::FToolBarSeparatorBlock(const FName& InExtensionHook)
-	: FMultiBlock( nullptr, nullptr, InExtensionHook, EMultiBlockType::ToolBarSeparator )
+	: FMultiBlock( nullptr, nullptr, InExtensionHook, EMultiBlockType::Separator )
 {
 }
 
@@ -18,7 +18,7 @@ FToolBarSeparatorBlock::FToolBarSeparatorBlock(const FName& InExtensionHook)
 
 void FToolBarSeparatorBlock::CreateMenuEntry(FMenuBuilder& MenuBuilder) const
 {
-	MenuBuilder.AddMenuSeparator();
+	MenuBuilder.AddSeparator();
 }
 
 
diff --git a/Engine/Source/Runtime/Slate/Public/Framework/MultiBox/MultiBox.h b/Engine/Source/Runtime/Slate/Public/Framework/MultiBox/MultiBox.h
--- a/Engine/Source/Runtime/Slate/Public/Framework/MultiBox/MultiBox.h
+++ b/Engine/Source/Runtime/Slate/Public/Framework/MultiBox/MultiBox.h
@@ -159,7 +159,7 @@ public:
 	 */
 	bool IsSeparator() const
 	{
-		return (Type == EMultiBlockType::MenuSeparator) || (Type == EMultiBlockType::ToolBarSeparator);
+		return Type == EMultiBlockType::Separator;
 	}
 
 	/**
diff --git a/Engine/Source/Runtime/Slate/Public/Framework/MultiBox/MultiBoxBuilder.h b/Engine/Source/Runtime/Slate/Public/Framework/MultiBox/MultiBoxBuilder.h
--- a/Engine/Source/Runtime/Slate/Public/Framework/MultiBox/MultiBoxBuilder.h
+++ b/Engine/Source/Runtime/Slate/Public/Framework/MultiBox/MultiBoxBuilder.h
@@ -242,11 +242,11 @@ public:
 	virtual TSharedRef< class SWidget > MakeWidget( FMultiBox::FOnMakeMultiBoxBuilderOverride* InMakeMultiBoxBuilderOverride = nullptr, uint32 MaxHeight = 1000) override;
 
 	/**
-	 * Adds a menu separator
+	 * Adds a separator
 	 */
 	void AddMenuSeparator(FName InExtensionHook = NAME_None);
+	void AddSeparator(FName InExtensionHook = NAME_None);
 	
-
 	/**
 	 * Starts a section on to the extender section hook stack
 	 * 
diff --git a/Engine/Source/Runtime/Slate/Public/Framework/MultiBox/MultiBoxDefs.h b/Engine/Source/Runtime/Slate/Public/Framework/MultiBox/MultiBoxDefs.h
--- a/Engine/Source/Runtime/Slate/Public/Framework/MultiBox/MultiBoxDefs.h
+++ b/Engine/Source/Runtime/Slate/Public/Framework/MultiBox/MultiBoxDefs.h
@@ -51,10 +51,9 @@ enum class EMultiBlockType : uint8
 	EditableText,
 	Heading,
 	MenuEntry,
-	MenuSeparator,
+	Separator,
 	ToolBarButton,
 	ToolBarComboButton,
-	ToolBarSeparator,
 	Widget,
 };
 
