Fix being able to select game mode pawn classes when it should be blocked.

diff --git a/Engine/Source/Developer/ToolMenus/Private/ToolMenus.cpp b/Engine/Source/Developer/ToolMenus/Private/ToolMenus.cpp
--- a/Engine/Source/Developer/ToolMenus/Private/ToolMenus.cpp
+++ b/Engine/Source/Developer/ToolMenus/Private/ToolMenus.cpp
@@ -918,6 +918,7 @@ void UToolMenus::PopulateMenuBuilder(FMenuBuilder& MenuBuilder, UToolMenu* MenuD
 			}
 
 			FUIAction UIAction = ConvertUIAction(Block, MenuData->Context);
+			const bool bUIActionIsSet = UIAction.ExecuteAction.IsBound() || UIAction.CanExecuteAction.IsBound() || UIAction.GetActionCheckState.IsBound() || UIAction.IsActionVisibleDelegate.IsBound();
 
 			TSharedPtr<SWidget> Widget;
 			if (Block.MakeWidget.IsBound())
@@ -979,7 +980,7 @@ void UToolMenus::PopulateMenuBuilder(FMenuBuilder& MenuBuilder, UToolMenu* MenuD
 					{
 						if (Widget.IsValid())
 						{
-							if (UIAction.IsBound())
+							if (bUIActionIsSet)
 							{
 								MenuBuilder.AddSubMenu(UIAction, Widget.ToSharedRef(), NewMenuDelegate, Block.bShouldCloseWindowAfterMenuSelection);
 							}
@@ -990,7 +991,7 @@ void UToolMenus::PopulateMenuBuilder(FMenuBuilder& MenuBuilder, UToolMenu* MenuD
 						}
 						else
 						{
-							if (UIAction.IsBound())
+							if (bUIActionIsSet)
 							{
 								MenuBuilder.AddSubMenu(
 									Block.Label,
