Disabled menu customization via config file as long as it conflicts with runtime menu customization
#rnx

diff --git a/Engine/Source/Developer/ToolMenus/Public/ToolMenus.h b/Engine/Source/Developer/ToolMenus/Public/ToolMenus.h
--- a/Engine/Source/Developer/ToolMenus/Public/ToolMenus.h
+++ b/Engine/Source/Developer/ToolMenus/Public/ToolMenus.h
@@ -363,11 +363,11 @@ private:
 
 private:
 
-	UPROPERTY(config, EditAnywhere, Category = Misc)
+	UPROPERTY(EditAnywhere, Category = Misc)
 	TArray<FCustomizedToolMenu> CustomizedMenus;
 
 	/* Allow substituting one menu for another during generate but not during find or extend */
-	UPROPERTY(config, EditAnywhere, Category = Misc)
+	UPROPERTY(EditAnywhere, Category = Misc)
 	TMap<FName, FName> MenuSubstitutionsDuringGenerate;
 
 	UPROPERTY()
