Submenus now have the option to collapse when only one item

diff --git a/Engine/Source/Developer/ToolMenus/Private/ToolMenus.cpp b/Engine/Source/Developer/ToolMenus/Private/ToolMenus.cpp
--- a/Engine/Source/Developer/ToolMenus/Private/ToolMenus.cpp
+++ b/Engine/Source/Developer/ToolMenus/Private/ToolMenus.cpp
@@ -72,6 +72,252 @@ FToolUIActionChoice::FToolUIActionChoice(const TSharedPtr< const FUICommandInfo
 	}
 }
 
+class FPopulateMenuBuilderWithToolMenuEntry
+{
+public:
+	FPopulateMenuBuilderWithToolMenuEntry(FMenuBuilder& InMenuBuilder, UToolMenu* InMenuData, FToolMenuSection& InSection, FToolMenuEntry& InBlock, bool bInAllowSubMenuCollapse) :
+		MenuBuilder(InMenuBuilder),
+		MenuData(InMenuData),
+		Section(InSection),
+		Block(InBlock),
+		BlockNameOverride(InBlock.Name),
+		bAllowSubMenuCollapse(bInAllowSubMenuCollapse),
+		bIsEditing(InMenuData->IsEditing())
+	{
+	}
+
+	void AddSubMenuEntryToMenuBuilder()
+	{
+		FName SubMenuFullName = UToolMenus::JoinMenuPaths(MenuData->MenuName, BlockNameOverride);
+		FNewMenuDelegate NewMenuDelegate;
+		bool bSubMenuAdded = false;
+
+		if (Block.SubMenuData.ConstructMenu.NewMenuLegacy.IsBound())
+		{
+			NewMenuDelegate = Block.SubMenuData.ConstructMenu.NewMenuLegacy;
+		}
+		else if (Block.SubMenuData.ConstructMenu.NewToolMenuWidget.IsBound())
+		{
+			// Full replacement of the widget shown when submenu is opened
+			FOnGetContent OnGetContent = UToolMenus::Get()->ConvertWidgetChoice(Block.SubMenuData.ConstructMenu, MenuData->Context);
+			if (OnGetContent.IsBound())
+			{
+				MenuBuilder.AddWrapperSubMenu(
+					Block.Label.Get(),
+					Block.ToolTip.Get(),
+					OnGetContent,
+					Block.Icon.Get()
+				);
+			}
+			bSubMenuAdded = true;
+		}
+		else if (BlockNameOverride == NAME_None)
+		{
+			if (Block.SubMenuData.ConstructMenu.NewToolMenu.IsBound())
+			{
+				// Blocks with no name cannot call PopulateSubMenu()
+				NewMenuDelegate = FNewMenuDelegate::CreateUObject(UToolMenus::Get(), &UToolMenus::PopulateSubMenuWithoutName, TWeakObjectPtr<UToolMenu>(MenuData), Block.SubMenuData.ConstructMenu.NewToolMenu);
+			}
+			else
+			{
+				UE_LOG(LogToolMenus, Warning, TEXT("Submenu that has no name is missing required delegate: %s"), *MenuData->MenuName.ToString());
+			}
+		}
+		else
+		{
+			if (Block.SubMenuData.bAutoCollapse && bAllowSubMenuCollapse)
+			{
+				// Preview the submenu to see if it should be collapsed
+				UToolMenu* GeneratedSubMenu = UToolMenus::Get()->GenerateSubMenu(MenuData, BlockNameOverride);
+				if (GeneratedSubMenu)
+				{
+					int32 NumSubMenuEntries = 0;
+					FToolMenuEntry* FirstEntry = nullptr;
+					for (FToolMenuSection& SubMenuSection : GeneratedSubMenu->Sections)
+					{
+						NumSubMenuEntries += SubMenuSection.Blocks.Num();
+						if (!FirstEntry && SubMenuSection.Blocks.Num() > 0)
+						{
+							FirstEntry = &SubMenuSection.Blocks[0];
+						}
+					}
+
+					if (NumSubMenuEntries == 1)
+					{
+						// Use bAllowSubMenuCollapse = false to avoid recursively collapsing a hierarchy of submenus that each contain one item
+						FPopulateMenuBuilderWithToolMenuEntry PopulateMenuBuilderWithToolMenuEntry(MenuBuilder, MenuData, Section, *FirstEntry, /* bAllowSubMenuCollapse= */ false);
+						PopulateMenuBuilderWithToolMenuEntry.SetBlockNameOverride(Block.Name);
+						PopulateMenuBuilderWithToolMenuEntry.Populate();
+						return;
+					}
+				}
+			}
+
+			NewMenuDelegate = FNewMenuDelegate::CreateUObject(UToolMenus::Get(), &UToolMenus::PopulateSubMenu, TWeakObjectPtr<UToolMenu>(MenuData), BlockNameOverride);
+		}
+
+		if (!bSubMenuAdded)
+		{
+			if (Widget.IsValid())
+			{
+				if (bUIActionIsSet)
+				{
+					MenuBuilder.AddSubMenu(UIAction, Widget.ToSharedRef(), NewMenuDelegate, Block.bShouldCloseWindowAfterMenuSelection);
+				}
+				else
+				{
+					MenuBuilder.AddSubMenu(Widget.ToSharedRef(), NewMenuDelegate, Block.SubMenuData.bOpenSubMenuOnClick, Block.bShouldCloseWindowAfterMenuSelection);
+				}
+			}
+			else
+			{
+				if (bUIActionIsSet)
+				{
+					MenuBuilder.AddSubMenu(
+						Block.Label,
+						Block.ToolTip,
+						NewMenuDelegate,
+						UIAction,
+						BlockNameOverride,
+						Block.UserInterfaceActionType,
+						Block.SubMenuData.bOpenSubMenuOnClick,
+						Block.Icon.Get(),
+						Block.bShouldCloseWindowAfterMenuSelection
+					);
+				}
+				else
+				{
+					MenuBuilder.AddSubMenu(
+						Block.Label,
+						Block.ToolTip,
+						NewMenuDelegate,
+						Block.SubMenuData.bOpenSubMenuOnClick,
+						Block.Icon.Get(),
+						Block.bShouldCloseWindowAfterMenuSelection,
+						BlockNameOverride
+					);
+				}
+			}
+		}
+	}
+
+	void AddStandardEntryToMenuBuilder()
+	{
+		if (Block.Command.IsValid())
+		{
+			bool bPopCommandList = false;
+			TSharedPtr<const FUICommandList> CommandListForAction;
+			if (Block.GetActionForCommand(MenuData->Context, CommandListForAction) != nullptr && CommandListForAction.IsValid())
+			{
+				MenuBuilder.PushCommandList(CommandListForAction.ToSharedRef());
+				bPopCommandList = true;
+			}
+			else
+			{
+				UE_LOG(LogToolMenus, Error, TEXT("UI command not found for menu entry: %s, menu: %s"), *BlockNameOverride.ToString(), *MenuData->MenuName.ToString());
+			}
+
+			MenuBuilder.AddMenuEntry(Block.Command, BlockNameOverride, LabelToDisplay, Block.ToolTip, Block.Icon.Get());
+
+			if (bPopCommandList)
+			{
+				MenuBuilder.PopCommandList();
+			}
+		}
+		else if (Block.ScriptObject)
+		{
+			UToolMenuEntryScript* ScriptObject = Block.ScriptObject;
+			const FSlateIcon Icon = ScriptObject->CreateIconAttribute(MenuData->Context).Get();
+			MenuBuilder.AddMenuEntry(ScriptObject->CreateLabelAttribute(MenuData->Context), ScriptObject->CreateToolTipAttribute(MenuData->Context), Icon, UIAction, ScriptObject->Data.Name, Block.UserInterfaceActionType, Block.TutorialHighlightName);
+		}
+		else
+		{
+			if (Widget.IsValid())
+			{
+				MenuBuilder.AddMenuEntry(UIAction, Widget.ToSharedRef(), BlockNameOverride, Block.ToolTip, Block.UserInterfaceActionType, Block.TutorialHighlightName);
+			}
+			else
+			{
+				MenuBuilder.AddMenuEntry(LabelToDisplay, Block.ToolTip, Block.Icon.Get(), UIAction, BlockNameOverride, Block.UserInterfaceActionType, Block.TutorialHighlightName);
+			}
+		}
+	}
+
+	void Populate()
+	{
+		if (Block.ConstructLegacy.IsBound())
+		{
+			if (!bIsEditing)
+			{
+				Block.ConstructLegacy.Execute(MenuBuilder, MenuData);
+			}
+
+			return;
+		}
+
+		UIAction = UToolMenus::ConvertUIAction(Block, MenuData->Context);
+		bUIActionIsSet = UIAction.ExecuteAction.IsBound() || UIAction.CanExecuteAction.IsBound() || UIAction.GetActionCheckState.IsBound() || UIAction.IsActionVisibleDelegate.IsBound();
+
+		if (Block.MakeWidget.IsBound())
+		{
+			Widget = Block.MakeWidget.Execute(MenuData->Context);
+		}
+
+		LabelToDisplay = Block.Label;
+		if (bIsEditing && (!Block.Label.IsSet() || Block.Label.Get().IsEmpty()))
+		{
+			LabelToDisplay = FText::FromName(BlockNameOverride);
+		}
+
+		if (Block.Type == EMultiBlockType::MenuEntry)
+		{
+			if (Block.IsSubMenu())
+			{
+				AddSubMenuEntryToMenuBuilder();
+			}
+			else
+			{
+				AddStandardEntryToMenuBuilder();
+			}
+		}
+		else if (Block.Type == EMultiBlockType::Separator)
+		{
+			MenuBuilder.AddMenuSeparator(BlockNameOverride);
+		}
+		else if (Block.Type == EMultiBlockType::Widget)
+		{
+			if (bIsEditing)
+			{
+				MenuBuilder.AddMenuEntry(LabelToDisplay, Block.ToolTip, Block.Icon.Get(), UIAction, BlockNameOverride, Block.UserInterfaceActionType, Block.TutorialHighlightName);
+			}
+			else
+			{
+				MenuBuilder.AddWidget(Widget.ToSharedRef(), LabelToDisplay.Get(), Block.WidgetData.bNoIndent, Block.WidgetData.bSearchable);
+			}
+		}
+		else
+		{
+			UE_LOG(LogToolMenus, Warning, TEXT("Menu '%s', item '%s', type not currently supported: %d"), *MenuData->MenuName.ToString(), *BlockNameOverride.ToString(), Block.Type);
+		}
+	};
+
+	void SetBlockNameOverride(const FName InBlockNameOverride) { BlockNameOverride = InBlockNameOverride; };
+
+private:
+	FMenuBuilder& MenuBuilder;
+	UToolMenu* MenuData;
+	FToolMenuSection& Section;
+	FToolMenuEntry& Block;
+	FName BlockNameOverride;
+
+	FUIAction UIAction;
+	TSharedPtr<SWidget> Widget;
+	TAttribute<FText> LabelToDisplay;
+	bool bAllowSubMenuCollapse;
+	bool bUIActionIsSet;
+	const bool bIsEditing;
+};
+
 UToolMenus::UToolMenus() :
 	bNextTickTimerIsSet(false),
 	bRefreshWidgetsNextTick(false),
@@ -908,179 +1154,8 @@ void UToolMenus::PopulateMenuBuilder(FMenuBuilder& MenuBuilder, UToolMenu* MenuD
 
 		for (FToolMenuEntry& Block : Section.Blocks)
 		{
-			if (Block.ConstructLegacy.IsBound())
-			{
-				if (!bIsEditing)
-				{
-					Block.ConstructLegacy.Execute(MenuBuilder, MenuData);
-				}
-				continue;
-			}
-
-			FUIAction UIAction = ConvertUIAction(Block, MenuData->Context);
-			const bool bUIActionIsSet = UIAction.ExecuteAction.IsBound() || UIAction.CanExecuteAction.IsBound() || UIAction.GetActionCheckState.IsBound() || UIAction.IsActionVisibleDelegate.IsBound();
-
-			TSharedPtr<SWidget> Widget;
-			if (Block.MakeWidget.IsBound())
-			{
-				Widget = Block.MakeWidget.Execute(MenuData->Context);
-			}
-
-			TAttribute<FText> LabelToDisplay = Block.Label;
-			if (bIsEditing && (!Block.Label.IsSet() || Block.Label.Get().IsEmpty()))
-			{
-				LabelToDisplay = FText::FromName(Block.Name);
-			}
-
-			if (Block.Type == EMultiBlockType::MenuEntry)
-			{
-				if (Block.IsSubMenu())
-				{
-					FName SubMenuFullName = JoinMenuPaths(MenuData->MenuName, Block.Name);
-					FNewMenuDelegate NewMenuDelegate;
-					bool bSubMenuAdded = false;
-
-					if (Block.SubMenuData.ConstructMenu.NewMenuLegacy.IsBound())
-					{
-						NewMenuDelegate = Block.SubMenuData.ConstructMenu.NewMenuLegacy;
-					}
-					else if (Block.SubMenuData.ConstructMenu.NewToolMenuWidget.IsBound())
-					{
-						// Full replacement of the widget shown when submenu is opened
-						FOnGetContent OnGetContent = ConvertWidgetChoice(Block.SubMenuData.ConstructMenu, MenuData->Context);
-						if (OnGetContent.IsBound())
-						{
-							MenuBuilder.AddWrapperSubMenu(
-								Block.Label.Get(),
-								Block.ToolTip.Get(),
-								OnGetContent,
-								Block.Icon.Get()
-							);
-						}
-						bSubMenuAdded = true;
-					}
-					else if (Block.Name == NAME_None)
-					{
-						if (Block.SubMenuData.ConstructMenu.NewToolMenu.IsBound())
-						{
-							// Blocks with no name cannot call PopulateSubMenu()
-							NewMenuDelegate = FNewMenuDelegate::CreateUObject(this, &UToolMenus::PopulateSubMenuWithoutName, TWeakObjectPtr<UToolMenu>(MenuData), Block.SubMenuData.ConstructMenu.NewToolMenu);
-						}
-						else
-						{
-							UE_LOG(LogToolMenus, Warning, TEXT("Submenu that has no name is missing required delegate: %s"), *MenuData->MenuName.ToString());
-						}
-					}
-					else
-					{
-						NewMenuDelegate = FNewMenuDelegate::CreateUObject(this, &UToolMenus::PopulateSubMenu, TWeakObjectPtr<UToolMenu>(MenuData), Block.Name);
-					}
-
-					if (!bSubMenuAdded)
-					{
-						if (Widget.IsValid())
-						{
-							if (bUIActionIsSet)
-							{
-								MenuBuilder.AddSubMenu(UIAction, Widget.ToSharedRef(), NewMenuDelegate, Block.bShouldCloseWindowAfterMenuSelection);
-							}
-							else
-							{
-								MenuBuilder.AddSubMenu(Widget.ToSharedRef(), NewMenuDelegate, Block.SubMenuData.bOpenSubMenuOnClick, Block.bShouldCloseWindowAfterMenuSelection);
-							}
-						}
-						else
-						{
-							if (bUIActionIsSet)
-							{
-								MenuBuilder.AddSubMenu(
-									Block.Label,
-									Block.ToolTip,
-									NewMenuDelegate,
-									UIAction,
-									Block.Name,
-									Block.UserInterfaceActionType,
-									Block.SubMenuData.bOpenSubMenuOnClick,
-									Block.Icon.Get(),
-									Block.bShouldCloseWindowAfterMenuSelection
-								);
-							}
-							else
-							{
-								MenuBuilder.AddSubMenu(
-									Block.Label,
-									Block.ToolTip,
-									NewMenuDelegate,
-									Block.SubMenuData.bOpenSubMenuOnClick,
-									Block.Icon.Get(),
-									Block.bShouldCloseWindowAfterMenuSelection,
-									Block.Name
-								);
-							}
-						}
-					}
-				}
-				else
-				{
-					if (Block.Command.IsValid())
-					{
-						bool bPopCommandList = false;
-						TSharedPtr<const FUICommandList> CommandListForAction;
-						if (Block.GetActionForCommand(MenuData->Context, CommandListForAction) != nullptr && CommandListForAction.IsValid())
-						{
-							MenuBuilder.PushCommandList(CommandListForAction.ToSharedRef());
-							bPopCommandList = true;
-						}
-						else
-						{
-							UE_LOG(LogToolMenus, Error, TEXT("UI command not found for menu entry: %s, menu: %s"), *Block.Name.ToString(), *MenuData->MenuName.ToString());
-						}
-
-						MenuBuilder.AddMenuEntry(Block.Command, Block.Name, LabelToDisplay, Block.ToolTip, Block.Icon.Get());
-
-						if (bPopCommandList)
-						{
-							MenuBuilder.PopCommandList();
-						}
-					}
-					else if (Block.ScriptObject)
-					{
-						UToolMenuEntryScript* ScriptObject = Block.ScriptObject;
-						const FSlateIcon Icon = ScriptObject->CreateIconAttribute(MenuData->Context).Get();
-						MenuBuilder.AddMenuEntry(ScriptObject->CreateLabelAttribute(MenuData->Context), ScriptObject->CreateToolTipAttribute(MenuData->Context), Icon, UIAction, ScriptObject->Data.Name, Block.UserInterfaceActionType, Block.TutorialHighlightName);
-					}
-					else
-					{
-						if (Widget.IsValid())
-						{
-							MenuBuilder.AddMenuEntry(UIAction, Widget.ToSharedRef(), Block.Name, Block.ToolTip, Block.UserInterfaceActionType, Block.TutorialHighlightName);
-						}
-						else
-						{
-							MenuBuilder.AddMenuEntry(LabelToDisplay, Block.ToolTip, Block.Icon.Get(), UIAction, Block.Name, Block.UserInterfaceActionType, Block.TutorialHighlightName);
-						}
-					}
-				}
-			}
-			else if (Block.Type == EMultiBlockType::Separator)
-			{
-				MenuBuilder.AddMenuSeparator(Block.Name);
-			}
-			else if (Block.Type == EMultiBlockType::Widget)
-			{
-				if (bIsEditing)
-				{
-					MenuBuilder.AddMenuEntry(LabelToDisplay, Block.ToolTip, Block.Icon.Get(), UIAction, Block.Name, Block.UserInterfaceActionType, Block.TutorialHighlightName);
-				}
-				else
-				{
-					MenuBuilder.AddWidget(Widget.ToSharedRef(), LabelToDisplay.Get(), Block.WidgetData.bNoIndent, Block.WidgetData.bSearchable);
-				}
-			}
-			else
-			{
-				UE_LOG(LogToolMenus, Warning, TEXT("Menu '%s', item '%s', type not currently supported: %d"), *MenuData->MenuName.ToString(), *Block.Name.ToString(), Block.Type);
-			}
+			FPopulateMenuBuilderWithToolMenuEntry PopulateMenuBuilderWithToolMenuEntry(MenuBuilder, MenuData, Section, Block, /* bAllowSubMenuCollapse= */ true);
+			PopulateMenuBuilderWithToolMenuEntry.Populate();
 		}
 
 		MenuBuilder.EndSection();
diff --git a/Engine/Source/Developer/ToolMenus/Public/ToolMenuEntry.h b/Engine/Source/Developer/ToolMenus/Public/ToolMenuEntry.h
--- a/Engine/Source/Developer/ToolMenus/Public/ToolMenuEntry.h
+++ b/Engine/Source/Developer/ToolMenus/Public/ToolMenuEntry.h
@@ -22,12 +22,15 @@ struct FToolMenuEntrySubMenuData
 public:
 	FToolMenuEntrySubMenuData() :
 		bIsSubMenu(false),
-		bOpenSubMenuOnClick(false)
+		bOpenSubMenuOnClick(false),
+		bAutoCollapse(false)
 	{
 	}
 
 	bool bIsSubMenu;
 	bool bOpenSubMenuOnClick;
+	/** Entry placed into the parent's menu when there is only one entry */
+	bool bAutoCollapse;
 	FNewToolMenuChoice ConstructMenu;
 };
 
@@ -162,6 +165,7 @@ private:
 
 	friend class UToolMenus;
 	friend class UToolMenuEntryExtensions;
+	friend class FPopulateMenuBuilderWithToolMenuEntry;
 
 	TAttribute<FText> Label;
 	TAttribute<FText> ToolTip;
diff --git a/Engine/Source/Developer/ToolMenus/Public/ToolMenuEntryScript.h b/Engine/Source/Developer/ToolMenus/Public/ToolMenuEntryScript.h
--- a/Engine/Source/Developer/ToolMenus/Public/ToolMenuEntryScript.h
+++ b/Engine/Source/Developer/ToolMenus/Public/ToolMenuEntryScript.h
@@ -158,6 +158,7 @@ private:
 
 	friend struct FToolMenuSection;
 	friend class UToolMenus;
+	friend class FPopulateMenuBuilderWithToolMenuEntry;
 
 	TAttribute<FText> CreateLabelAttribute(FToolMenuContext& Context);
 
diff --git a/Engine/Source/Developer/ToolMenus/Public/ToolMenus.h b/Engine/Source/Developer/ToolMenus/Public/ToolMenus.h
--- a/Engine/Source/Developer/ToolMenus/Public/ToolMenus.h
+++ b/Engine/Source/Developer/ToolMenus/Public/ToolMenus.h
@@ -296,6 +296,7 @@ public:
 	//~ End UObject Interface
 
 private:
+	friend class FPopulateMenuBuilderWithToolMenuEntry;
 
 	/** Create a finalized menu that combines given hierarchy array that will generate a widget. Advanced special use cases only. */
 	UToolMenu* GenerateMenuFromHierarchy(const TArray<UToolMenu*>& Hierarchy, const FToolMenuContext& InMenuContext);
