Fix for level editor viewport menu commands not working after changing layout.
#rnx

diff --git a/Engine/Source/Developer/ToolMenus/Private/ToolMenus.cpp b/Engine/Source/Developer/ToolMenus/Private/ToolMenus.cpp
--- a/Engine/Source/Developer/ToolMenus/Private/ToolMenus.cpp
+++ b/Engine/Source/Developer/ToolMenus/Private/ToolMenus.cpp
@@ -322,6 +322,7 @@ UToolMenus::UToolMenus() :
 	bNextTickTimerIsSet(false),
 	bRefreshWidgetsNextTick(false),
 	bCleanupStaleWidgetsNextTick(false),
+	bCleanupStaleWidgetsNextTickGC(false),
 	bEditMenusMode(false)
 {
 }
@@ -1852,9 +1853,15 @@ void UToolMenus::SetNextTickTimer()
 	}
 }
 
-void UToolMenus::CleanupStaleWidgetsNextTick()
+void UToolMenus::CleanupStaleWidgetsNextTick(bool bGarbageCollect)
 {
 	bCleanupStaleWidgetsNextTick = true;
+
+	if (bGarbageCollect)
+	{
+		bCleanupStaleWidgetsNextTickGC = true;
+	}
+
 	SetNextTickTimer();
 }
 
@@ -1870,6 +1877,7 @@ void UToolMenus::HandleNextTick()
 	{
 		CleanupStaleWidgets();
 		bCleanupStaleWidgetsNextTick = false;
+		bCleanupStaleWidgetsNextTickGC = false;
 
 		if (bRefreshWidgetsNextTick)
 		{
@@ -1894,6 +1902,7 @@ void UToolMenus::HandleNextTick()
 
 void UToolMenus::CleanupStaleWidgets()
 {
+	bool bModified = false;
 	for (auto WidgetsForMenuNameIt = GeneratedMenuWidgets.CreateIterator(); WidgetsForMenuNameIt; ++WidgetsForMenuNameIt)
 	{
 		FGeneratedToolMenuWidgets& WidgetsForMenuName = WidgetsForMenuNameIt->Value;
@@ -1902,15 +1911,22 @@ void UToolMenus::CleanupStaleWidgets()
 		{
 			if (!Instance->Widget.IsValid())
 			{
+				bModified = true;
 				Instance.RemoveCurrent();
 			}
 		}
 
 		if (WidgetsForMenuName.Instances.Num() == 0)
 		{
+			bModified = true;
 			WidgetsForMenuNameIt.RemoveCurrent();
 		}
 	}
+
+	if (bModified && bCleanupStaleWidgetsNextTickGC && !IsAsyncLoading())
+	{
+		CollectGarbage(GARBAGE_COLLECTION_KEEPFLAGS);
+	}
 }
 
 bool UToolMenus::RefreshMenuWidget(const FName InName)
diff --git a/Engine/Source/Developer/ToolMenus/Public/ToolMenus.h b/Engine/Source/Developer/ToolMenus/Public/ToolMenus.h
--- a/Engine/Source/Developer/ToolMenus/Public/ToolMenus.h
+++ b/Engine/Source/Developer/ToolMenus/Public/ToolMenus.h
@@ -266,6 +266,9 @@ public:
 	/* Remove substitute one menu for another during generate */
 	void RemoveSubstitutionDuringGenerate(const FName InMenu);
 
+	/** Release references to UObjects of widgets that have been deleted. Combines multiple requests in one frame together for improved performance. */
+	void CleanupStaleWidgetsNextTick(bool bGarbageCollect = false);
+
 	/** Displaying extension points is for debugging menus */
 	DECLARE_DELEGATE_RetVal(bool, FShouldDisplayExtensionPoints);
 	FShouldDisplayExtensionPoints ShouldDisplayExtensionPoints;
@@ -304,9 +307,6 @@ private:
 	/** Sets a timer to be called next engine tick so that multiple repeated actions can be combined together. */
 	void SetNextTickTimer();
 
-	/** Release references to UObjects of widgets that have been deleted. Combines multiple requests in one frame together for improved performance. */
-	void CleanupStaleWidgetsNextTick();
-
 	/** Release references to UObjects of widgets that have been deleted */
 	void CleanupStaleWidgets();
 
@@ -387,6 +387,7 @@ private:
 	bool bNextTickTimerIsSet;
 	bool bRefreshWidgetsNextTick;
 	bool bCleanupStaleWidgetsNextTick;
+	bool bCleanupStaleWidgetsNextTickGC;
 	bool bEditMenusMode;
 
 	static UToolMenus* Singleton;
diff --git a/Engine/Source/Editor/LevelEditor/Private/SLevelViewport.cpp b/Engine/Source/Editor/LevelEditor/Private/SLevelViewport.cpp
--- a/Engine/Source/Editor/LevelEditor/Private/SLevelViewport.cpp
+++ b/Engine/Source/Editor/LevelEditor/Private/SLevelViewport.cpp
@@ -73,6 +73,7 @@
 #include "EditorLevelUtils.h"
 #include "Engine/LevelStreaming.h"
 #include "Editor/WorldBrowser/Public/WorldBrowserModule.h"
+#include "ToolMenus.h"
 
 static const FName LevelEditorName("LevelEditor");
 
@@ -3934,6 +3935,7 @@ void SLevelViewport::OnSetViewportConfiguration(FName ConfigurationName)
 			GCurrentLevelEditingViewportClient = nullptr;
 			ViewportTabPinned->SetViewportConfiguration(ConfigurationName);
 			FSlateApplication::Get().DismissAllMenus();
+			UToolMenus::Get()->CleanupStaleWidgetsNextTick(true);
 		}
 	}
 }
