Added menu context clean-up functions

These are called when a menu is closed, and can be used to clean-up context object data without needing to run a GC

diff --git a/Engine/Source/Developer/ToolMenus/Private/ToolMenu.cpp b/Engine/Source/Developer/ToolMenus/Private/ToolMenu.cpp
--- a/Engine/Source/Developer/ToolMenus/Private/ToolMenu.cpp
+++ b/Engine/Source/Developer/ToolMenus/Private/ToolMenu.cpp
@@ -12,6 +12,7 @@
 
 UToolMenu::UToolMenu() :
 	MenuType(EMultiBoxType::Menu)
+	, bShouldCleanupContextOnDestroy(true)
 	, bShouldCloseWindowAfterMenuSelection(true)
 	, bCloseSelfOnly(false)
 	, bSearchable(false)
@@ -501,6 +502,14 @@ void UToolMenu::UpdateMenuCustomizationFromMultibox(const TSharedRef<const FMult
 	}
 }
 
+void UToolMenu::OnMenuDestroyed()
+{
+	if (bShouldCleanupContextOnDestroy && !SubMenuParent)
+	{
+		Context.CleanupObjects();
+	}
+}
+
 TArray<const UToolMenu*> UToolMenu::GetSubMenuChain() const
 {
 	TArray<const UToolMenu*> SubMenuChain;
diff --git a/Engine/Source/Developer/ToolMenus/Private/ToolMenuContext.cpp b/Engine/Source/Developer/ToolMenus/Private/ToolMenuContext.cpp
--- a/Engine/Source/Developer/ToolMenus/Private/ToolMenuContext.cpp
+++ b/Engine/Source/Developer/ToolMenus/Private/ToolMenuContext.cpp
@@ -11,20 +11,27 @@
 #include "Internationalization/Internationalization.h"
 
 
-FToolMenuContext::FToolMenuContext() : bIsEditing(false)
+FToolMenuContext::FToolMenuContext(UObject* InContext)
 {
-
+	if (InContext)
+	{
+		ContextObjects.Add(InContext);
+	}
 }
 
-FToolMenuContext::FToolMenuContext(UObject* InContext) : bIsEditing(false)
+FToolMenuContext::FToolMenuContext(UObject* InContext, FContextObjectCleanup&& InCleanup)
 {
 	if (InContext)
 	{
 		ContextObjects.Add(InContext);
+		if (InCleanup)
+		{
+			ContextObjectCleanupFuncs.Add(InContext, MoveTemp(InCleanup));
+		}
 	}
 }
 
-FToolMenuContext::FToolMenuContext(TSharedPtr<FUICommandList> InCommandList, TSharedPtr<FExtender> InExtender, UObject* InContext) : bIsEditing(false)
+FToolMenuContext::FToolMenuContext(TSharedPtr<FUICommandList> InCommandList, TSharedPtr<FExtender> InExtender, UObject* InContext)
 {
 	if (InContext)
 	{
@@ -142,9 +149,27 @@ void FToolMenuContext::AddObject(UObject* InObject)
 	ContextObjects.AddUnique(InObject);
 }
 
+void FToolMenuContext::AddObject(UObject* InObject, FContextObjectCleanup&& InCleanup)
+{
+	ContextObjects.AddUnique(InObject);
+	if (InCleanup)
+	{
+		ContextObjectCleanupFuncs.Add(InObject, MoveTemp(InCleanup));
+	}
+}
+
+void FToolMenuContext::CleanupObjects()
+{
+	for (const TTuple<TObjectPtr<UObject>, FContextObjectCleanup>& CleanupPair : ContextObjectCleanupFuncs)
+	{
+		CleanupPair.Value(CleanupPair.Key);
+	}
+}
+
 void FToolMenuContext::Empty()
 {
 	ContextObjects.Empty();
+	ContextObjectCleanupFuncs.Empty();
 	CommandLists.Empty();
 	CommandList.Reset();
 	Extenders.Empty();
diff --git a/Engine/Source/Developer/ToolMenus/Private/ToolMenus.cpp b/Engine/Source/Developer/ToolMenus/Private/ToolMenus.cpp
--- a/Engine/Source/Developer/ToolMenus/Private/ToolMenus.cpp
+++ b/Engine/Source/Developer/ToolMenus/Private/ToolMenus.cpp
@@ -1824,6 +1824,7 @@ TSharedRef<SWidget> UToolMenus::GenerateWidget(UToolMenu* GeneratedMenu)
 	// Store a copy so that we can call 'Refresh' on menus not in the database
 	FGeneratedToolMenuWidget& GeneratedMenuWidget = WidgetsForMenuName.Instances.AddDefaulted_GetRef();
 	GeneratedMenuWidget.GeneratedMenu = DuplicateObject<UToolMenu>(GeneratedMenu, this, MakeUniqueObjectName(this, UToolMenus::StaticClass(), FName("MenuForRefresh")));
+	GeneratedMenuWidget.GeneratedMenu->bShouldCleanupContextOnDestroy = true;
 	// Copy native properties that serialize does not
 	GeneratedMenuWidget.GeneratedMenu->Context = GeneratedMenu->Context;
 	GeneratedMenuWidget.GeneratedMenu->StyleSet = GeneratedMenu->StyleSet;
@@ -1851,7 +1852,6 @@ TSharedRef<SWidget> UToolMenus::GenerateWidget(UToolMenu* GeneratedMenu)
 		}
 
 		MenuBuilder.SetExtendersEnabled(GeneratedMenu->bExtendersEnabled);
-		MenuBuilder.GetMultiBox()->WeakToolMenu = GeneratedMenu;
 		PopulateMenuBuilder(MenuBuilder, GeneratedMenu);
 		if (GeneratedMenu->ModifyBlockWidgetAfterMake.IsBound())
 		{
@@ -2048,6 +2048,7 @@ bool UToolMenus::RefreshMenuWidget(const FName InName, FGeneratedToolMenuWidget&
 	}
 
 	// Regenerate menu from database
+	GeneratedMenuWidget.GeneratedMenu->bShouldCleanupContextOnDestroy = false; // The new menu will do this
 	UToolMenu* GeneratedMenu = GenerateMenu(InName, GeneratedMenuWidget.GeneratedMenu->Context);
 	GeneratedMenuWidget.GeneratedMenu = GeneratedMenu;
 
diff --git a/Engine/Source/Developer/ToolMenus/Public/ToolMenu.h b/Engine/Source/Developer/ToolMenus/Public/ToolMenu.h
--- a/Engine/Source/Developer/ToolMenus/Public/ToolMenu.h
+++ b/Engine/Source/Developer/ToolMenus/Public/ToolMenu.h
@@ -75,6 +75,7 @@ public:
 	virtual FCustomizedToolMenu* AddMenuCustomization() const override;
 	virtual FCustomizedToolMenuHierarchy GetMenuCustomizationHierarchy() const override;
 	virtual void UpdateMenuCustomizationFromMultibox(const TSharedRef<const FMultiBox>& InMultiBox) override;
+	virtual void OnMenuDestroyed() override;
 	//~ End UToolMenuBase Interface
 
 	TArray<FName> GetMenuHierarchyNames(bool bIncludeSubMenuRoot) const;
@@ -129,6 +130,9 @@ public:
 	UPROPERTY(EditAnywhere, BlueprintReadWrite, Category = "Tool Menus")
 	EMultiBoxType MenuType;
 
+	UPROPERTY(Transient)
+	bool bShouldCleanupContextOnDestroy;
+
 	UPROPERTY(EditAnywhere, BlueprintReadWrite, Category = "Tool Menus")
 	bool bShouldCloseWindowAfterMenuSelection;
 
diff --git a/Engine/Source/Developer/ToolMenus/Public/ToolMenuContext.h b/Engine/Source/Developer/ToolMenus/Public/ToolMenuContext.h
--- a/Engine/Source/Developer/ToolMenus/Public/ToolMenuContext.h
+++ b/Engine/Source/Developer/ToolMenus/Public/ToolMenuContext.h
@@ -3,6 +3,7 @@
 #pragma once
 
 #include "CoreMinimal.h"
+#include "Containers/SortedMap.h"
 
 #include "ToolMenuContext.generated.h"
 
@@ -33,8 +34,11 @@ struct TOOLMENUS_API FToolMenuContext
 	GENERATED_BODY()
 public:
 
-	FToolMenuContext();
+	using FContextObjectCleanup = TFunction<void(UObject*)>;
+
+	FToolMenuContext() = default;
 	FToolMenuContext(UObject* InContext);
+	FToolMenuContext(UObject* InContext, FContextObjectCleanup&& InCleanup);
 	FToolMenuContext(TSharedPtr<FUICommandList> InCommandList, TSharedPtr<FExtender> InExtender = TSharedPtr<FExtender>(), UObject* InContext = nullptr);
 
 	template <typename TContextType>
@@ -54,15 +58,7 @@ public:
 	template <typename TContextType>
 	TContextType* Find() const
 	{
-		for (UObject* Object : ContextObjects)
-		{
-			if (TContextType* Result = Cast<TContextType>(Object))
-			{
-				return Result;
-			}
-		}
-
-		return nullptr;
+		return FindContext<TContextType>();
 	}
 
 	UObject* FindByClass(UClass* InClass) const;
@@ -74,31 +70,35 @@ public:
 
 	void AddExtender(const TSharedPtr<FExtender>& InExtender);
 	TSharedPtr<FExtender> GetAllExtenders() const;
-	void ReplaceExtenders(const TSharedPtr<FExtender>& InExtender);
 	void ResetExtenders();
 
 	void AppendObjects(const TArray<UObject*>& InObjects);
 	void AddObject(UObject* InObject);
+	void AddObject(UObject* InObject, FContextObjectCleanup&& InCleanup);
+
+	void CleanupObjects();
 
 	friend class UToolMenus;
 	friend class UToolMenu;
 	friend struct FToolMenuEntry;
 
 	bool IsEditing() const { return bIsEditing; }
-
-	bool bIsEditing;
+	void SetIsEditing(bool InIsEditing) { bIsEditing = InIsEditing; }
 
 private:
 
 	void Empty();
 
+	bool bIsEditing = false;
+
 	UPROPERTY()
 	TArray<TObjectPtr<UObject>> ContextObjects;
 
+	TSortedMap<TObjectPtr<UObject>, FContextObjectCleanup> ContextObjectCleanupFuncs;
+
 	TArray<TSharedPtr<FUICommandList>> CommandLists;
 
 	TSharedPtr<FUICommandList> CommandList;
 
 	TArray<TSharedPtr<FExtender>> Extenders;
 };
-
diff --git a/Engine/Source/Editor/ToolMenusEditor/Private/SEditToolMenuDialog.cpp b/Engine/Source/Editor/ToolMenusEditor/Private/SEditToolMenuDialog.cpp
--- a/Engine/Source/Editor/ToolMenusEditor/Private/SEditToolMenuDialog.cpp
+++ b/Engine/Source/Editor/ToolMenusEditor/Private/SEditToolMenuDialog.cpp
@@ -768,8 +768,8 @@ void SEditToolMenuDialog::OnMenuNamesSelectionChanged(TSharedPtr<FName> InEntry,
 	}
 	
 	FToolMenuContext NewMenuContext = SourceMenu->Context;
-	NewMenuContext.bIsEditing = true;
-		
+	NewMenuContext.SetIsEditing(true);
+	
 	TArray<const UToolMenu*> SubMenuChain = SourceMenu->GetSubMenuChain();
 	if (!SourceMenu->SubMenuParent)
 	{
diff --git a/Engine/Source/Runtime/Slate/Private/Framework/MultiBox/MultiBox.cpp b/Engine/Source/Runtime/Slate/Private/Framework/MultiBox/MultiBox.cpp
--- a/Engine/Source/Runtime/Slate/Private/Framework/MultiBox/MultiBox.cpp
+++ b/Engine/Source/Runtime/Slate/Private/Framework/MultiBox/MultiBox.cpp
@@ -206,6 +206,10 @@ FMultiBox::FMultiBox(const EMultiBoxType InType, FMultiBoxCustomization InCustom
 
 FMultiBox::~FMultiBox()
 {
+	if (UToolMenuBase* ToolMenu = GetToolMenu())
+	{
+		ToolMenu->OnMenuDestroyed();
+	}
 }
 
 TSharedRef<FMultiBox> FMultiBox::Create( const EMultiBoxType InType, FMultiBoxCustomization InCustomization, const bool bInShouldCloseWindowAfterMenuSelection )
diff --git a/Engine/Source/Runtime/Slate/Public/Framework/MultiBox/ToolMenuBase.h b/Engine/Source/Runtime/Slate/Public/Framework/MultiBox/ToolMenuBase.h
--- a/Engine/Source/Runtime/Slate/Public/Framework/MultiBox/ToolMenuBase.h
+++ b/Engine/Source/Runtime/Slate/Public/Framework/MultiBox/ToolMenuBase.h
@@ -120,5 +120,6 @@ public:
 	virtual FCustomizedToolMenu* AddMenuCustomization() const { return nullptr; }
 	virtual FCustomizedToolMenuHierarchy GetMenuCustomizationHierarchy() const { return FCustomizedToolMenuHierarchy(); }
 	virtual void UpdateMenuCustomizationFromMultibox(const TSharedRef<const FMultiBox>& InMultiBox) {}
+	virtual void OnMenuDestroyed() {}
 };
 
