Make tool menu context class finder forward found class type to BP return type

diff --git a/Engine/Source/Developer/ToolMenus/Private/ToolMenusBlueprintLibrary.cpp b/Engine/Source/Developer/ToolMenus/Private/ToolMenusBlueprintLibrary.cpp
--- a/Engine/Source/Developer/ToolMenus/Private/ToolMenusBlueprintLibrary.cpp
+++ b/Engine/Source/Developer/ToolMenus/Private/ToolMenusBlueprintLibrary.cpp
@@ -44,9 +44,9 @@ void UToolMenuEntryExtensions::BreakToolMenuOwner(const FToolMenuOwner& InValue,
 	Name = InValue.TryGetName();
 }
 
-UObject* UToolMenuContextExtensions::FindByClass(const FToolMenuContext& Context, UClass* InClass)
+UObject* UToolMenuContextExtensions::FindByClass(const FToolMenuContext& Context, TSubclassOf<UObject> InClass)
 {
-	return Context.FindByClass(InClass);
+	return Context.FindByClass(*InClass);
 }
 
 void UToolMenuEntryExtensions::SetLabel(UPARAM(ref) FToolMenuEntry& Target, const FText& Label)
diff --git a/Engine/Source/Developer/ToolMenus/Private/ToolMenusBlueprintLibrary.h b/Engine/Source/Developer/ToolMenus/Private/ToolMenusBlueprintLibrary.h
--- a/Engine/Source/Developer/ToolMenus/Private/ToolMenusBlueprintLibrary.h
+++ b/Engine/Source/Developer/ToolMenus/Private/ToolMenusBlueprintLibrary.h
@@ -10,6 +10,7 @@
 #include "ToolMenuMisc.h"
 #include "ToolMenuEntryScript.h"
 #include "Misc/Attribute.h"
+#include "Templates/SubclassOf.h"
 
 #include "ToolMenusBlueprintLibrary.generated.h"
 
@@ -20,8 +21,8 @@ class UToolMenuContextExtensions : public UObject
 
 public:
 
-	UFUNCTION(BlueprintCallable, Category = "Tool Menus", meta = (ScriptMethod))
-	static UObject* FindByClass(const FToolMenuContext& Context, UClass* InClass);
+	UFUNCTION(BlueprintCallable, Category = "Tool Menus", meta = (ScriptMethod, DeterminesOutputType = "InClass"))
+	static UObject* FindByClass(const FToolMenuContext& Context, TSubclassOf<UObject> InClass);
 };
 
 UCLASS()
