Adding conditionals and logging to find main menu crashes on Mac

diff --git a/Engine/Source/Developer/ToolMenus/Private/ToolMenus.cpp b/Engine/Source/Developer/ToolMenus/Private/ToolMenus.cpp
--- a/Engine/Source/Developer/ToolMenus/Private/ToolMenus.cpp
+++ b/Engine/Source/Developer/ToolMenus/Private/ToolMenus.cpp
@@ -491,6 +491,10 @@ void UToolMenus::ListAllParents(const FName InName, TArray<FName>& AllParents)
 
 void UToolMenus::AssembleMenuSection(UToolMenu* GeneratedMenu, const UToolMenu* Other, FToolMenuSection* DestSection, const FToolMenuSection& OtherSection)
 {
+	if (!DestSection)
+	{
+		UE_LOG(LogToolMenus, Warning, TEXT("Trying to add to invalid section for menu: %s, section: %s. Default section info will be used instead."), *OtherSection.Owner.TryGetName().ToString(), *OtherSection.Name.ToString());
+	}
 	// Build list of blocks in expected order including blocks created by construct delegates
 	TArray<FToolMenuEntry> RemainingBlocks;
 	TArray<FToolMenuEntry> BlocksToAddLast;
@@ -518,7 +522,14 @@ void UToolMenus::AssembleMenuSection(UToolMenu* GeneratedMenu, const UToolMenu*
 			{
 				break;
 			}
-			ConstructedEntries->Context = DestSection->Context;
+			if (DestSection)
+			{
+				ConstructedEntries->Context = DestSection->Context;
+			}
+			else
+			{
+				ConstructedEntries->Context = FToolMenuContext();
+			}
 		}
 
 		TArray<FToolMenuEntry> GeneratedEntries;
@@ -532,18 +543,36 @@ void UToolMenus::AssembleMenuSection(UToolMenu* GeneratedMenu, const UToolMenu*
 			{
 				if (NumIterations++ > 5000)
 				{
-					UE_LOG(LogToolMenus, Warning, TEXT("Possible infinite loop for menu: %s, section: %s, block: %s"), *Other->MenuName.ToString(), *OtherSection.Name.ToString(), *Block.Name.ToString());
+					FName MenuName = OtherSection.Owner.TryGetName();
+
+					if (Other)
+					{
+						MenuName = Other->MenuName;
+					}
+					UE_LOG(LogToolMenus, Warning, TEXT("Possible infinite loop for menu: %s, section: %s, block: %s"), *MenuName.ToString(), *OtherSection.Name.ToString(), *Block.Name.ToString());
 					break;
 				}
 				
 				ConstructedEntries->Sections.Reset();
 				if (GeneratedEntry.IsScriptObjectDynamicConstruct())
 				{
-					GeneratedEntry.ScriptObject->ConstructMenuEntry(ConstructedEntries, DestSection->Name, DestSection->Context);
+					FName SectionName;
+					FToolMenuContext SectionContext;
+					if (DestSection)
+					{
+						SectionName = DestSection->Name;
+						SectionContext = DestSection->Context;
+					}
+					GeneratedEntry.ScriptObject->ConstructMenuEntry(ConstructedEntries, SectionName, SectionContext);
 				}
 				else
 				{
-					FToolMenuSection& ConstructedSection = ConstructedEntries->AddSection(DestSection->Name);
+					FName SectionName;
+					if (DestSection)
+					{
+						SectionName = DestSection->Name;
+					}
+					FToolMenuSection& ConstructedSection = ConstructedEntries->AddSection(SectionName);
 					ConstructedSection.Context = ConstructedEntries->Context;
 					GeneratedEntry.Construct.Execute(ConstructedSection);
 				}
