Adding one more missing conditional

diff --git a/Engine/Source/Developer/ToolMenus/Private/ToolMenus.cpp b/Engine/Source/Developer/ToolMenus/Private/ToolMenus.cpp
--- a/Engine/Source/Developer/ToolMenus/Private/ToolMenus.cpp
+++ b/Engine/Source/Developer/ToolMenus/Private/ToolMenus.cpp
@@ -617,31 +617,35 @@ void UToolMenus::AssembleMenuSection(UToolMenu* GeneratedMenu, const UToolMenu*
 
 	RemainingBlocks.Append(BlocksToAddLast);
 
-	// Repeatedly loop because insert location may not exist until later in list
-	while (RemainingBlocks.Num() > 0)
+	// Only do this loop if there is a section to insert into. We need to early-out here or it will be an infinite loop
+	if (DestSection)
 	{
-		int32 NumHandled = 0;
-		for (int32 i = 0; i < RemainingBlocks.Num(); ++i)
+		// Repeatedly loop because insert location may not exist until later in list
+		while (RemainingBlocks.Num() > 0)
 		{
-			FToolMenuEntry& Block = RemainingBlocks[i];
-			int32 DestIndex = DestSection->FindBlockInsertIndex(Block);
-			if (DestIndex != INDEX_NONE)
+			int32 NumHandled = 0;
+			for (int32 i = 0; i < RemainingBlocks.Num(); ++i)
 			{
-				DestSection->Blocks.Insert(Block, DestIndex);
-				RemainingBlocks.RemoveAt(i);
-				--i;
-				++NumHandled;
-				// Restart loop because items earlier in the list may need to attach to this block
-				break;
+				FToolMenuEntry& Block = RemainingBlocks[i];
+				int32 DestIndex = DestSection->FindBlockInsertIndex(Block);
+				if (DestIndex != INDEX_NONE)
+				{
+					DestSection->Blocks.Insert(Block, DestIndex);
+					RemainingBlocks.RemoveAt(i);
+					--i;
+					++NumHandled;
+					// Restart loop because items earlier in the list may need to attach to this block
+					break;
+				}
 			}
-		}
-		if (NumHandled == 0)
-		{
-			for (const FToolMenuEntry& Block : RemainingBlocks)
+			if (NumHandled == 0)
 			{
-				UE_LOG(LogToolMenus, Warning, TEXT("Menu item not found: '%s' for insert: '%s'"), *Block.InsertPosition.Name.ToString(), *Block.Name.ToString());
+				for (const FToolMenuEntry& Block : RemainingBlocks)
+				{
+					UE_LOG(LogToolMenus, Warning, TEXT("Menu item not found: '%s' for insert: '%s'"), *Block.InsertPosition.Name.ToString(), *Block.Name.ToString());
+				}
+				break;
 			}
-			break;
 		}
 	}
 }
