Fix tool menus holding onto widgets after being closed and before gc

diff --git a/Engine/Source/Developer/ToolMenus/Private/ToolMenu.cpp b/Engine/Source/Developer/ToolMenus/Private/ToolMenu.cpp
--- a/Engine/Source/Developer/ToolMenus/Private/ToolMenu.cpp
+++ b/Engine/Source/Developer/ToolMenus/Private/ToolMenu.cpp
@@ -513,6 +513,8 @@ void UToolMenu::OnMenuDestroyed()
 	{
 		Context.CleanupObjects();
 	}
+
+	Empty();
 }
 
 TArray<const UToolMenu*> UToolMenu::GetSubMenuChain() const
