Disable early cleanup of menus while a few edge cases can be debugged and fixed

diff --git a/Engine/Source/Developer/ToolMenus/Private/ToolMenu.cpp b/Engine/Source/Developer/ToolMenus/Private/ToolMenu.cpp
--- a/Engine/Source/Developer/ToolMenus/Private/ToolMenu.cpp
+++ b/Engine/Source/Developer/ToolMenus/Private/ToolMenu.cpp
@@ -514,7 +514,7 @@ void UToolMenu::OnMenuDestroyed()
 		Context.CleanupObjects();
 	}
 
-	Empty();
+	//Empty();
 }
 
 TArray<const UToolMenu*> UToolMenu::GetSubMenuChain() const
