Fix menu losing name when calling certain API

diff --git a/Engine/Source/Developer/ToolMenus/Private/ToolMenuSection.cpp b/Engine/Source/Developer/ToolMenus/Private/ToolMenuSection.cpp
--- a/Engine/Source/Developer/ToolMenus/Private/ToolMenuSection.cpp
+++ b/Engine/Source/Developer/ToolMenus/Private/ToolMenuSection.cpp
@@ -93,7 +93,7 @@ FToolMenuEntry& FToolMenuSection::AddMenuEntry(const FName InNameOverride, const
 	return AddEntry(FToolMenuEntry::InitMenuEntry(InCommand, InLabelOverride, InToolTipOverride, InIconOverride, InTutorialHighlightName, InNameOverride));
 }
 
-FToolMenuEntry& FToolMenuSection::AddMenuEntryWithCommandList(const TSharedPtr< const FUICommandInfo >& InCommand, const TSharedPtr< const FUICommandList >& InCommandList, const TAttribute<FText>& InLabelOverride, const TAttribute<FText>& InToolTipOverride, const TAttribute<FSlateIcon>& InIconOverride, const FName InTutorialHighlightName, const FName InNameOverride)
+FToolMenuEntry& FToolMenuSection::AddMenuEntryWithCommandList(const TSharedPtr< const FUICommandInfo >& InCommand, const TSharedPtr< const FUICommandList >& InCommandList, const TAttribute<FText>& InLabelOverride, const TAttribute<FText>& InToolTipOverride, const TAttribute<FSlateIcon>& InIconOverride, const FName InTutorialHighlightName, const TOptional<FName> InNameOverride)
 {
 	return AddEntry(FToolMenuEntry::InitMenuEntryWithCommandList(InCommand, InCommandList, InLabelOverride, InToolTipOverride, InIconOverride, InTutorialHighlightName, InNameOverride));
 }
diff --git a/Engine/Source/Developer/ToolMenus/Public/ToolMenuSection.h b/Engine/Source/Developer/ToolMenus/Public/ToolMenuSection.h
--- a/Engine/Source/Developer/ToolMenus/Public/ToolMenuSection.h
+++ b/Engine/Source/Developer/ToolMenus/Public/ToolMenuSection.h
@@ -37,7 +37,7 @@ public:
 	FToolMenuEntry& AddMenuEntry(const FName InName, const TAttribute<FText>& InLabel, const TAttribute<FText>& InToolTip, const TAttribute<FSlateIcon>& InIcon, const FToolUIActionChoice& InAction, const EUserInterfaceActionType UserInterfaceActionType = EUserInterfaceActionType::Button, const FName InTutorialHighlightName = NAME_None);
 	FToolMenuEntry& AddMenuEntry(const TSharedPtr< const FUICommandInfo >& InCommand, const TAttribute<FText>& InLabelOverride = TAttribute<FText>(), const TAttribute<FText>& InToolTipOverride = TAttribute<FText>(), const TAttribute<FSlateIcon>& InIconOverride = TAttribute<FSlateIcon>(), const FName InTutorialHighlightName = NAME_None, const TOptional<FName> InNameOverride = TOptional<FName>());
 	FToolMenuEntry& AddMenuEntry(const FName InNameOverride, const TSharedPtr< const FUICommandInfo >& InCommand, const TAttribute<FText>& InLabelOverride = TAttribute<FText>(), const TAttribute<FText>& InToolTipOverride = TAttribute<FText>(), const TAttribute<FSlateIcon>& InIconOverride = TAttribute<FSlateIcon>(), const FName InTutorialHighlightName = NAME_None);
-	FToolMenuEntry& AddMenuEntryWithCommandList(const TSharedPtr<const FUICommandInfo>& InCommand, const TSharedPtr<const FUICommandList>& InCommandList, const TAttribute<FText>& InLabelOverride = TAttribute<FText>(), const TAttribute<FText>& InToolTipOverride = TAttribute<FText>(), const TAttribute<FSlateIcon>& InIconOverride = TAttribute<FSlateIcon>(), const FName InTutorialHighlightName = NAME_None, const FName InNameOverride = NAME_None);
+	FToolMenuEntry& AddMenuEntryWithCommandList(const TSharedPtr<const FUICommandInfo>& InCommand, const TSharedPtr<const FUICommandList>& InCommandList, const TAttribute<FText>& InLabelOverride = TAttribute<FText>(), const TAttribute<FText>& InToolTipOverride = TAttribute<FText>(), const TAttribute<FSlateIcon>& InIconOverride = TAttribute<FSlateIcon>(), const FName InTutorialHighlightName = NAME_None, const TOptional<FName> InNameOverride = TOptional<FName>());
 
 	FToolMenuEntry& AddDynamicEntry(const FName InName, const FNewToolMenuSectionDelegate& InConstruct);
 	FToolMenuEntry& AddDynamicEntry(const FName InName, const FNewToolMenuDelegateLegacy& InConstruct);
