ToolMenus - Updated debug warning message to display enum value as string when trying to add menu entry to toolbar

diff --git a/Engine/Source/Developer/ToolMenus/Private/ToolMenus.cpp b/Engine/Source/Developer/ToolMenus/Private/ToolMenus.cpp
--- a/Engine/Source/Developer/ToolMenus/Private/ToolMenus.cpp
+++ b/Engine/Source/Developer/ToolMenus/Private/ToolMenus.cpp
@@ -682,6 +682,7 @@ void UToolMenus::AssembleMenu(UToolMenu* GeneratedMenu, const UToolMenu* Other)
 				break;
 			}
 			ConstructedSections->Context = GeneratedMenu->Context;
+			ConstructedSections->MenuType = GeneratedMenu->MenuType;
 		}
 
 		TArray<FToolMenuSection> GeneratedSections;
@@ -1394,7 +1395,7 @@ void UToolMenus::PopulateToolBarBuilder(FToolBarBuilder& ToolBarBuilder, UToolMe
 			}
 			else
 			{
-				UE_LOG(LogToolMenus, Warning, TEXT("Toolbar '%s', item '%s', type not currently supported: %d"), *MenuData->MenuName.ToString(), *Block.Name.ToString(), Block.Type);
+				UE_LOG(LogToolMenus, Warning, TEXT("Toolbar '%s', item '%s', Toolbars do not support: %s"), *MenuData->MenuName.ToString(), *Block.Name.ToString(), *UEnum::GetValueAsString(Block.Type));
 			}
 
 			ToolBarBuilder.EndStyleOverride();
