Fixed custom FOnGetContent not working for UToolMenu submenus

diff --git a/Engine/Source/Developer/ToolMenus/Private/ToolMenus.cpp b/Engine/Source/Developer/ToolMenus/Private/ToolMenus.cpp
--- a/Engine/Source/Developer/ToolMenus/Private/ToolMenus.cpp
+++ b/Engine/Source/Developer/ToolMenus/Private/ToolMenus.cpp
@@ -112,6 +112,17 @@ public:
 			}
 			bSubMenuAdded = true;
 		}
+		else if (Block.SubMenuData.ConstructMenu.OnGetContent.IsBound())
+		{
+			// Full replacement of the widget shown when submenu is opened
+			MenuBuilder.AddWrapperSubMenu(
+				Block.Label.Get(),
+				Block.ToolTip.Get(),
+				Block.SubMenuData.ConstructMenu.OnGetContent,
+				Block.Icon.Get()
+			);
+			bSubMenuAdded = true;
+		}
 		else if (BlockNameOverride == NAME_None)
 		{
 			if (Block.SubMenuData.ConstructMenu.NewToolMenu.IsBound())
