Minor follow up adjustment of recent ToolMenu submenu change

diff --git a/Engine/Source/Developer/ToolMenus/Private/ToolMenus.cpp b/Engine/Source/Developer/ToolMenus/Private/ToolMenus.cpp
--- a/Engine/Source/Developer/ToolMenus/Private/ToolMenus.cpp
+++ b/Engine/Source/Developer/ToolMenus/Private/ToolMenus.cpp
@@ -97,7 +97,7 @@ public:
 		{
 			NewMenuDelegate = Block.SubMenuData.ConstructMenu.NewMenuLegacy;
 		}
-		else if (Block.SubMenuData.ConstructMenu.NewToolMenuWidget.IsBound())
+		else if (Block.SubMenuData.ConstructMenu.NewToolMenuWidget.IsBound() || Block.SubMenuData.ConstructMenu.OnGetContent.IsBound())
 		{
 			// Full replacement of the widget shown when submenu is opened
 			FOnGetContent OnGetContent = UToolMenus::Get()->ConvertWidgetChoice(Block.SubMenuData.ConstructMenu, MenuData->Context);
@@ -112,17 +112,6 @@ public:
 			}
 			bSubMenuAdded = true;
 		}
-		else if (Block.SubMenuData.ConstructMenu.OnGetContent.IsBound())
-		{
-			// Full replacement of the widget shown when submenu is opened
-			MenuBuilder.AddWrapperSubMenu(
-				Block.Label.Get(),
-				Block.ToolTip.Get(),
-				Block.SubMenuData.ConstructMenu.OnGetContent,
-				Block.Icon.Get()
-			);
-			bSubMenuAdded = true;
-		}
 		else if (BlockNameOverride == NAME_None)
 		{
 			if (Block.SubMenuData.ConstructMenu.NewToolMenu.IsBound())
