Deprecate AddOptionsDropdown
#rnx

diff --git a/Engine/Source/Developer/ToolMenus/Public/ToolMenuEntry.h b/Engine/Source/Developer/ToolMenus/Public/ToolMenuEntry.h
--- a/Engine/Source/Developer/ToolMenus/Public/ToolMenuEntry.h
+++ b/Engine/Source/Developer/ToolMenus/Public/ToolMenuEntry.h
@@ -139,7 +139,9 @@ struct TOOLMENUS_API FToolMenuEntry
 
 	void SetCommandList(const TSharedPtr<const FUICommandList>& InCommandList);
 
+	UE_DEPRECATED(5.1, "AddOptionsDropdown is deprecated. Use InitComboButton with bSimpleComboBox=true")
 	void AddOptionsDropdown(FUIAction InAction, const FOnGetContent InMenuContentGenerator, const TAttribute<FText>& InToolTip = TAttribute<FText>());
+
 	void AddKeybindFromCommand(const TSharedPtr< const FUICommandInfo >& InCommand);
 
 	bool IsCommandKeybindOnly() const;
