Menus - Add function to unregister all runtime customizations belonging to a specific owner name
#rnx

diff --git a/Engine/Source/Developer/ToolMenus/Private/ToolMenus.cpp b/Engine/Source/Developer/ToolMenus/Private/ToolMenus.cpp
--- a/Engine/Source/Developer/ToolMenus/Private/ToolMenus.cpp
+++ b/Engine/Source/Developer/ToolMenus/Private/ToolMenus.cpp
@@ -2347,50 +2347,75 @@ void UToolMenus::RemoveEntry(const FName MenuName, const FName InSection, const
 
 void UToolMenus::UnregisterOwnerInternal(FToolMenuOwner InOwner)
 {
-	if (InOwner != FToolMenuOwner())
+	if (InOwner == FToolMenuOwner())
 	{
-		for (auto It = Menus.CreateIterator(); It; ++It)
+		return;
+	}
+
+	bool bNeedsRefresh = false;
+
+	for (const TPair<FName, TObjectPtr<UToolMenu>>& Pair : Menus)
+	{
+		UToolMenu* Menu = Pair.Value;
+		for (int32 SectionIndex = Menu->Sections.Num() - 1; SectionIndex >= 0; --SectionIndex)
 		{
-			bool bNeedsRefresh = false;
+			FToolMenuSection& Section = Menu->Sections[SectionIndex];
+			if (Section.RemoveEntriesByOwner(InOwner) > 0)
+			{
+				bNeedsRefresh = true;
+			}
 
-			UToolMenu* Menu = It->Value;
-			for (int32 SectionIndex = Menu->Sections.Num() - 1; SectionIndex >=0; --SectionIndex)
+			if (Section.Owner == InOwner)
 			{
-				FToolMenuSection& Section = Menu->Sections[SectionIndex];
-				if (Section.RemoveEntriesByOwner(InOwner) > 0)
+				if (Section.Construct.IsBound())
 				{
+					Section.Construct = FNewSectionConstructChoice();
 					bNeedsRefresh = true;
 				}
 
-				if (Section.Owner == InOwner)
+				if (Section.ToolMenuSectionDynamic)
 				{
-					if (Section.Construct.IsBound())
-					{
-						Section.Construct = FNewSectionConstructChoice();
-						bNeedsRefresh = true;
-					}
-
-					if (Section.ToolMenuSectionDynamic)
-					{
-						Section.ToolMenuSectionDynamic = nullptr;
-						bNeedsRefresh = true;
-					}
+					Section.ToolMenuSectionDynamic = nullptr;
+					bNeedsRefresh = true;
+				}
 
-					if (Section.Blocks.Num() == 0)
-					{
-						Menu->Sections.RemoveAt(SectionIndex, 1, false);
-						bNeedsRefresh = true;
-					}
+				if (Section.Blocks.Num() == 0)
+				{
+					Menu->Sections.RemoveAt(SectionIndex, 1, false);
+					bNeedsRefresh = true;
 				}
 			}
+		}
+	}
 
-			// Refresh any widgets that are currently displayed to the user
-			if (bNeedsRefresh)
-			{
-				RefreshAllWidgets();
-			}
+	// Refresh any widgets that are currently displayed to the user
+	if (bNeedsRefresh)
+	{
+		RefreshAllWidgets();
+	}
+}
+
+void UToolMenus::UnregisterRuntimeMenuCustomizationOwner(FName InOwnerName)
+{
+	if (InOwnerName.IsNone())
+	{
+		return;
+	}
+
+	bool bNeedsRefresh = false;
+	for (FCustomizedToolMenu& CustomizedToolMenu : RuntimeCustomizedMenus)
+	{
+		if (CustomizedToolMenu.MenuPermissions.UnregisterOwner(InOwnerName))
+		{
+			bNeedsRefresh = true;
 		}
 	}
+
+	// Refresh any widgets that are currently displayed to the user
+	if (bNeedsRefresh)
+	{
+		RefreshAllWidgets();
+	}
 }
 
 FToolMenuOwner UToolMenus::CurrentOwner() const
diff --git a/Engine/Source/Developer/ToolMenus/Public/ToolMenus.h b/Engine/Source/Developer/ToolMenus/Public/ToolMenus.h
--- a/Engine/Source/Developer/ToolMenus/Public/ToolMenus.h
+++ b/Engine/Source/Developer/ToolMenus/Public/ToolMenus.h
@@ -254,6 +254,9 @@ public:
 	/** Find or add runtime customization settings for a menu */
 	FCustomizedToolMenu* AddRuntimeMenuCustomization(const FName InName);
 
+	/** Unregister runtime customization settings for a specific owner name */
+	void UnregisterRuntimeMenuCustomizationOwner(const FName InOwnerName);
+
 	/** Generates sub menu by entry name in the given generated menu parent */
 	UToolMenu* GenerateSubMenu(const UToolMenu* InGeneratedParent, const FName InBlockName);
 
