Add additional info to a log error to find out which UI Command doesn't have a UI Action.

diff --git a/Engine/Source/Developer/ToolMenus/Private/ToolMenus.cpp b/Engine/Source/Developer/ToolMenus/Private/ToolMenus.cpp
--- a/Engine/Source/Developer/ToolMenus/Private/ToolMenus.cpp
+++ b/Engine/Source/Developer/ToolMenus/Private/ToolMenus.cpp
@@ -219,7 +219,10 @@ public:
 			}
 			else
 			{
-				UE_LOG(LogToolMenus, Error, TEXT("UI command not found for menu entry: %s, menu: %s"), *BlockNameOverride.ToString(), *MenuData->MenuName.ToString());
+				UE_LOG(LogToolMenus, Error, TEXT("UI command not found for menu entry: %s[%s], menu: %s"),
+					*BlockNameOverride.ToString(), 
+					**FTextInspector::GetSourceString(LabelToDisplay.Get()),
+					*MenuData->MenuName.ToString());
 			}
 
 			MenuBuilder.AddMenuEntry(Block.Command, BlockNameOverride, LabelToDisplay, Block.ToolTip, Block.Icon.Get());
