[ToolMenus] Exposed StyleSet in UToolMenu

diff --git a/Engine/Source/Developer/ToolMenus/Private/ToolMenu.cpp b/Engine/Source/Developer/ToolMenus/Private/ToolMenu.cpp
--- a/Engine/Source/Developer/ToolMenus/Private/ToolMenu.cpp
+++ b/Engine/Source/Developer/ToolMenus/Private/ToolMenu.cpp
@@ -58,7 +58,18 @@ FReply UToolMenu::OnKeyDown(const FGeometry& MyGeometry, const FKeyEvent& InKeyE
 	return FReply::Unhandled();
 }
 
+const ISlateStyle* UToolMenu::GetStyleSet() const
+{
+	return StyleSet;
+}
 
+void UToolMenu::SetStyleSet(const ISlateStyle* InStyleSet)
+{
+	if (InStyleSet && InStyleSet != StyleSet)
+	{
+		StyleSet = InStyleSet;
+	}
+}
 
 void UToolMenu::InitGeneratedCopy(const UToolMenu* Source, const FName InMenuName, const FToolMenuContext* InContext)
 {
diff --git a/Engine/Source/Developer/ToolMenus/Public/ToolMenu.h b/Engine/Source/Developer/ToolMenus/Public/ToolMenu.h
--- a/Engine/Source/Developer/ToolMenus/Public/ToolMenu.h
+++ b/Engine/Source/Developer/ToolMenus/Public/ToolMenu.h
@@ -93,6 +93,9 @@ public:
 
 	FReply OnKeyDown(const FGeometry& MyGeometry, const FKeyEvent& InKeyEvent);
 
+	const ISlateStyle* GetStyleSet() const;
+	void SetStyleSet(const ISlateStyle* InStyleSet);
+
 	friend class UToolMenus;
 
 private:
