[ToolMenus] Added FToolMenuEntry::InputBindingLabel to allow for explicit input binding hint without requiring a UICommand

#rnx

diff --git a/Engine/Source/Developer/ToolMenus/Private/ToolMenus.cpp b/Engine/Source/Developer/ToolMenus/Private/ToolMenus.cpp
--- a/Engine/Source/Developer/ToolMenus/Private/ToolMenus.cpp
+++ b/Engine/Source/Developer/ToolMenus/Private/ToolMenus.cpp
@@ -236,17 +236,47 @@ public:
 		{
 			UToolMenuEntryScript* ScriptObject = Block.ScriptObject;
 			const FSlateIcon Icon = ScriptObject->CreateIconAttribute(MenuData->Context).Get();
-			MenuBuilder.AddMenuEntry(ScriptObject->CreateLabelAttribute(MenuData->Context), ScriptObject->CreateToolTipAttribute(MenuData->Context), Icon, UIAction, ScriptObject->Data.Name, Block.UserInterfaceActionType, Block.TutorialHighlightName);
+			
+			FMenuEntryParams MenuEntryParams;
+			MenuEntryParams.LabelOverride = ScriptObject->CreateLabelAttribute(MenuData->Context);
+			MenuEntryParams.ToolTipOverride = ScriptObject->CreateToolTipAttribute(MenuData->Context);
+			MenuEntryParams.IconOverride = Icon;
+			MenuEntryParams.DirectActions = UIAction;
+			MenuEntryParams.ExtensionHook = ScriptObject->Data.Name;
+			MenuEntryParams.UserInterfaceActionType = Block.UserInterfaceActionType;
+			MenuEntryParams.TutorialHighlightName = Block.TutorialHighlightName;
+			MenuEntryParams.InputBindingOverride = Block.InputBindingLabel;
+
+			MenuBuilder.AddMenuEntry(MenuEntryParams);
 		}
 		else
 		{
 			if (Widget.IsValid())
 			{
-				MenuBuilder.AddMenuEntry(UIAction, Widget.ToSharedRef(), BlockNameOverride, Block.ToolTip, Block.UserInterfaceActionType, Block.TutorialHighlightName);
+				FMenuEntryParams MenuEntryParams;
+				MenuEntryParams.DirectActions = UIAction;
+				MenuEntryParams.EntryWidget = Widget.ToSharedRef();
+				MenuEntryParams.ExtensionHook = BlockNameOverride;
+				MenuEntryParams.ToolTipOverride = Block.ToolTip;
+				MenuEntryParams.UserInterfaceActionType = Block.UserInterfaceActionType;
+				MenuEntryParams.TutorialHighlightName = Block.TutorialHighlightName;
+				MenuEntryParams.InputBindingOverride = Block.InputBindingLabel;
+
+				MenuBuilder.AddMenuEntry(MenuEntryParams);
 			}
 			else
 			{
-				MenuBuilder.AddMenuEntry(LabelToDisplay, Block.ToolTip, Block.Icon.Get(), UIAction, BlockNameOverride, Block.UserInterfaceActionType, Block.TutorialHighlightName);
+				FMenuEntryParams MenuEntryParams;
+				MenuEntryParams.LabelOverride = LabelToDisplay;
+				MenuEntryParams.ToolTipOverride = Block.ToolTip;
+				MenuEntryParams.IconOverride = Block.Icon.Get();
+				MenuEntryParams.DirectActions = UIAction;
+				MenuEntryParams.ExtensionHook = BlockNameOverride;
+				MenuEntryParams.UserInterfaceActionType = Block.UserInterfaceActionType;
+				MenuEntryParams.TutorialHighlightName = Block.TutorialHighlightName;
+				MenuEntryParams.InputBindingOverride = Block.InputBindingLabel;
+				
+				MenuBuilder.AddMenuEntry(MenuEntryParams);
 			}
 		}
 	}
@@ -306,7 +336,17 @@ public:
 		{
 			if (bIsEditing)
 			{
-				MenuBuilder.AddMenuEntry(LabelToDisplay, Block.ToolTip, Block.Icon.Get(), UIAction, BlockNameOverride, Block.UserInterfaceActionType, Block.TutorialHighlightName);
+				FMenuEntryParams MenuEntryParams;
+				MenuEntryParams.LabelOverride = LabelToDisplay;
+				MenuEntryParams.ToolTipOverride = Block.ToolTip;
+				MenuEntryParams.IconOverride = Block.Icon.Get();
+				MenuEntryParams.DirectActions = UIAction;
+				MenuEntryParams.ExtensionHook = BlockNameOverride;
+				MenuEntryParams.UserInterfaceActionType = Block.UserInterfaceActionType;
+				MenuEntryParams.TutorialHighlightName = Block.TutorialHighlightName;
+				MenuEntryParams.InputBindingOverride = Block.InputBindingLabel;
+				
+				MenuBuilder.AddMenuEntry(MenuEntryParams);
 			}
 			else
 			{
diff --git a/Engine/Source/Developer/ToolMenus/Public/ToolMenuEntry.h b/Engine/Source/Developer/ToolMenus/Public/ToolMenuEntry.h
--- a/Engine/Source/Developer/ToolMenus/Public/ToolMenuEntry.h
+++ b/Engine/Source/Developer/ToolMenus/Public/ToolMenuEntry.h
@@ -204,6 +204,7 @@ public:
 	TAttribute<FText> Label;
 	TAttribute<FText> ToolTip;
 	TAttribute<FSlateIcon> Icon;
+	TAttribute<FText> InputBindingLabel;
 
 private:
 
