[ToolMenus] Overloaded FindOrAdd functions with other Add params

#rnx

diff --git a/Engine/Source/Developer/ToolMenus/Private/ToolMenu.cpp b/Engine/Source/Developer/ToolMenus/Private/ToolMenu.cpp
--- a/Engine/Source/Developer/ToolMenus/Private/ToolMenu.cpp
+++ b/Engine/Source/Developer/ToolMenus/Private/ToolMenu.cpp
@@ -308,6 +308,19 @@ FToolMenuSection& UToolMenu::FindOrAddSection(const FName SectionName)
 	return AddSection(SectionName);
 }
 
+FToolMenuSection& UToolMenu::FindOrAddSection(
+	const FName SectionName,
+	const TAttribute<FText>& InLabel,
+	const FToolMenuInsert InPosition)
+{
+	if (FToolMenuSection* FoundSection = FindSection(SectionName))
+	{
+		return *FoundSection;
+	}
+
+	return AddSection(SectionName, InLabel, InPosition);
+}
+
 void UToolMenu::RemoveSection(const FName SectionName)
 {
 	Sections.RemoveAll([SectionName](const FToolMenuSection& Section) { return Section.Name == SectionName; });
diff --git a/Engine/Source/Developer/ToolMenus/Public/ToolMenu.h b/Engine/Source/Developer/ToolMenus/Public/ToolMenu.h
--- a/Engine/Source/Developer/ToolMenus/Public/ToolMenu.h
+++ b/Engine/Source/Developer/ToolMenus/Public/ToolMenu.h
@@ -52,6 +52,8 @@ public:
 	FToolMenuSection* FindSection(const FName SectionName);
 
 	FToolMenuSection& FindOrAddSection(const FName SectionName);
+	
+	FToolMenuSection& FindOrAddSection(const FName SectionName, const TAttribute< FText >& InLabel, const FToolMenuInsert InPosition = FToolMenuInsert());
 
 	FName GetMenuName() const { return MenuName; }
 
