Add Documentation Strings to User-Facing ToolMenus Code

This adds in-code documentation strings to a few more ToolMenus APIs. The goal is not full coverage but to hit the most commonly used bits and the ones new users most likely run into first.

diff --git a/Engine/Source/Developer/ToolMenus/Public/ToolMenu.h b/Engine/Source/Developer/ToolMenus/Public/ToolMenu.h
--- a/Engine/Source/Developer/ToolMenus/Public/ToolMenu.h
+++ b/Engine/Source/Developer/ToolMenus/Public/ToolMenu.h
@@ -13,6 +13,11 @@
 
 #include "ToolMenu.generated.h"
 
+/**
+ * Represents a menu in the ToolMenus system.
+ *
+ * An instance of this class is returned by basic APIs such as UToolMenus::RegisterMenu and UToolMenus::ExtendMenu.
+ **/
 UCLASS(BlueprintType)
 class TOOLMENUS_API UToolMenu : public UToolMenuBase
 {
@@ -46,13 +51,45 @@ public:
 
 	FToolMenuSection& AddSection(const FName SectionName, const TAttribute< FText >& InLabel = TAttribute<FText>(), const FToolMenuInsert InPosition = FToolMenuInsert());
 
-	/** Callback is passed sandboxed empty menu to add as many sections as desired. Sections merged into the final menu afterwards. */
+	/**
+	 * Build a section in a callback.
+	 *
+	 * The callback is passed a sandboxed empty menu to add as many sections as desired. Sections are merged into the final menu afterwards.
+	 *
+	 * @param SectionName The FName of the section the callback will build.
+	 * @param InConstruct The callback that will be run to build the menu as the user opens it.
+	 * @param InPosition The position of the section within its parent menu.
+	 * @return A reference to the created section.
+	 */
 	FToolMenuSection& AddDynamicSection(const FName SectionName, const FNewSectionConstructChoice& InConstruct, const FToolMenuInsert InPosition = FToolMenuInsert());
 
+	/**
+	 * Finds an existing section using the provided FName.
+	 *
+	 * @param SectionName The section name to find. Provide NAME_None to access the unnamed section of this menu.
+	 * @return A pointer to the found section or nullptr if none was found.
+	 */
 	FToolMenuSection* FindSection(const FName SectionName);
 
+	/**
+	 * Finds an existing section using the provided FName or, if not found, creates a new one and assigns it that FName.
+	 *
+	 * @param SectionName The section name to find or add. Provide NAME_None to access the unnamed section of this menu.
+	 * @return A reference to the found or newly created section.
+	 */
 	FToolMenuSection& FindOrAddSection(const FName SectionName);
-	
+
+	/**
+	 * Finds an existing section using the provided FName or, if not found, creates a new one and assigns it that FName.
+	 *
+	 * In case the the section did not exist and a new one was created, then the provided label and optional
+	 * position is also set on the newly created section.
+	 *
+	 * @param SectionName The section name to find or add. Provide NAME_None to access the unnamed section of this menu.
+	 * @param InLabel The label to set on the section if a new one is created. If an existing section is returned this label is ignored.
+	 * @param InPosition The position to set on the section if a new one is created. If an existing section is returned this position is ignored.
+	 * @return A reference to the found or newly created section.
+	 */
 	FToolMenuSection& FindOrAddSection(const FName SectionName, const TAttribute< FText >& InLabel, const FToolMenuInsert InPosition = FToolMenuInsert());
 
 	FName GetMenuName() const { return MenuName; }
diff --git a/Engine/Source/Developer/ToolMenus/Public/ToolMenuEntry.h b/Engine/Source/Developer/ToolMenus/Public/ToolMenuEntry.h
--- a/Engine/Source/Developer/ToolMenus/Public/ToolMenuEntry.h
+++ b/Engine/Source/Developer/ToolMenus/Public/ToolMenuEntry.h
@@ -99,6 +99,11 @@ struct FToolMenuCustomWidgetContext
 	FName StyleName;
 };
 
+/**
+ * Represents entries in menus such as buttons, checkboxes, and sub-menus.
+ *
+ * Many entries are created for you via the methods of FToolMenuSection, such as FToolMenuSection::AddMenuEntry.
+ */
 USTRUCT(BlueprintType)
 struct TOOLMENUS_API FToolMenuEntry
 {
diff --git a/Engine/Source/Developer/ToolMenus/Public/ToolMenuOwner.h b/Engine/Source/Developer/ToolMenus/Public/ToolMenuOwner.h
--- a/Engine/Source/Developer/ToolMenus/Public/ToolMenuOwner.h
+++ b/Engine/Source/Developer/ToolMenus/Public/ToolMenuOwner.h
@@ -6,7 +6,12 @@
 
 #include "ToolMenuOwner.generated.h"
 
-
+/**
+ * Represents the owner of a menu. Can be used when registering menus
+ * to later unregistering all menus created by a specified owner.
+ *
+ * See UToolMenus::UnregisterOwnerByName and FToolMenuOwnerScoped which relate to this.
+ */
 USTRUCT(BlueprintType, meta=(HasNativeBreak="/Script/ToolMenus.ToolMenuEntryExtensions.BreakToolMenuOwner", HasNativeMake="/Script/ToolMenus.ToolMenuEntryExtensions.MakeToolMenuOwner"))
 struct TOOLMENUS_API FToolMenuOwner
 {
diff --git a/Engine/Source/Developer/ToolMenus/Public/ToolMenuSection.h b/Engine/Source/Developer/ToolMenus/Public/ToolMenuSection.h
--- a/Engine/Source/Developer/ToolMenus/Public/ToolMenuSection.h
+++ b/Engine/Source/Developer/ToolMenus/Public/ToolMenuSection.h
@@ -34,6 +34,18 @@ public:
 
 	FToolMenuEntry& AddEntryObject(UToolMenuEntryScript* InObject);
 	FToolMenuEntry& AddEntry(const FToolMenuEntry& Args);
+
+	/**
+	 * Add a menu entry to this section.
+	 *
+	 * @param InName The FName of the added entry used to identify it.
+	 * @param InLabel The user-visible label of this entry.
+	 * @param InToolTip The user-visible tooltip of this entry.
+	 * @param InIcon The user-visible icon of this entry.
+	 * @param InAction The delegates to call to make this entry interactive. Supports more than just on-click delegates. See FToolUIActionChoice for more information.
+	 * @param UserInterfaceActionType The type of user interface element (such as a checkbox) to use for this entry. See EUserInterfaceActionType for more options.
+	 * @return A reference to the added menu entry.
+	 */
 	FToolMenuEntry& AddMenuEntry(const FName InName, const TAttribute<FText>& InLabel, const TAttribute<FText>& InToolTip, const TAttribute<FSlateIcon>& InIcon, const FToolUIActionChoice& InAction, const EUserInterfaceActionType UserInterfaceActionType = EUserInterfaceActionType::Button, const FName InTutorialHighlightName = NAME_None);
 	FToolMenuEntry& AddMenuEntry(const TSharedPtr< const FUICommandInfo >& InCommand, const TAttribute<FText>& InLabelOverride = TAttribute<FText>(), const TAttribute<FText>& InToolTipOverride = TAttribute<FText>(), const TAttribute<FSlateIcon>& InIconOverride = TAttribute<FSlateIcon>(), const FName InTutorialHighlightName = NAME_None, const TOptional<FName> InNameOverride = TOptional<FName>());
 	FToolMenuEntry& AddMenuEntry(const FName InNameOverride, const TSharedPtr< const FUICommandInfo >& InCommand, const TAttribute<FText>& InLabelOverride = TAttribute<FText>(), const TAttribute<FText>& InToolTipOverride = TAttribute<FText>(), const TAttribute<FSlateIcon>& InIconOverride = TAttribute<FSlateIcon>(), const FName InTutorialHighlightName = NAME_None);
@@ -46,6 +58,20 @@ public:
 	FToolMenuEntry& AddMenuSeparator(const FName InName);
 	FToolMenuEntry& AddSeparator(const FName InName);
 
+	/**
+	 * Add a sub-menu to this section.
+	 *
+	 * @param InName The FName to use for identifying this sub-menu. Must be unique within the section.
+	 * @param InLabel The user-visible label to use for this sub-menu.
+	 * @param InToolTip The user-visible tooltip to use for this sub-menu.
+	 * @param InMakeMenu The delegate to call to construct the sub-menu. This delegate is called when the user clicks the sub-menu entry.
+	 * @param InAction The delegates to use to make this menu entry interactive.
+	 * @param InUserInterfaceActionType The type of user interface element (such as a checkbox) to use for this entry.
+	 * @param bInOpenSubMenuOnClick Whether or not the sub-menu entry opens when it is hovered or only when clicked.
+	 * @param InIcon The user-visible icon to use for this sub-menu.
+	 * @param bShouldCloseWindowAfterMenuSelection Whether or not the sub-menu should close after one of its entries is clicked.
+	 * @return A reference to the menu entry for the added sub-menu.
+	 */
 	FToolMenuEntry& AddSubMenu(const FName InName, const TAttribute<FText>& InLabel, const TAttribute<FText>& InToolTip, const FNewToolMenuChoice& InMakeMenu, const FToolUIActionChoice& InAction, const EUserInterfaceActionType InUserInterfaceActionType, bool bInOpenSubMenuOnClick = false, const TAttribute<FSlateIcon>& InIcon = TAttribute<FSlateIcon>(), const bool bShouldCloseWindowAfterMenuSelection = true);
 	FToolMenuEntry& AddSubMenu(const FName InName, const TAttribute<FText>& InLabel, const TAttribute<FText>& InToolTip, const FNewToolMenuChoice& InMakeMenu, bool bInOpenSubMenuOnClick = false, const TAttribute<FSlateIcon>& InIcon = TAttribute<FSlateIcon>(), const bool bShouldCloseWindowAfterMenuSelection = true, const FName InTutorialHighlightName = NAME_None);
 	FToolMenuEntry& AddSubMenu(const FName InName, const FToolUIActionChoice& InAction, const TSharedRef<SWidget>& InWidget, const FNewToolMenuChoice& InMakeMenu, bool bShouldCloseWindowAfterMenuSelection = true);
diff --git a/Engine/Source/Developer/ToolMenus/Public/ToolMenus.h b/Engine/Source/Developer/ToolMenus/Public/ToolMenus.h
--- a/Engine/Source/Developer/ToolMenus/Public/ToolMenus.h
+++ b/Engine/Source/Developer/ToolMenus/Public/ToolMenus.h
@@ -416,6 +416,10 @@ private:
 	static TOptional<FDelegateHandle> InternalStartupCallbackHandle;
 };
 
+/**
+ * Sets the owner for all menus created until the end of the current scope (with support for nested scopes).
+ * Combines well with UToolMenus::UnregisterOwnerByName.
+ */
 struct FToolMenuOwnerScoped
 {
 	FToolMenuOwnerScoped(const FToolMenuOwner InOwner) : Owner(InOwner)
