Fix crash in UToolMenus::GenerateWidget(UToolMenu*)

* Fix a crash in UToolMenus::GenerateWidget(UToolMenu*) that could happen when UToolMenus::AddReferencedObjects was called while a legacy menu-building delegate was running.
* Add a unit test that triggers the crash if the fix is not applied.

diff --git a/Engine/Source/Developer/ToolMenus/Private/ToolMenus.cpp b/Engine/Source/Developer/ToolMenus/Private/ToolMenus.cpp
--- a/Engine/Source/Developer/ToolMenus/Private/ToolMenus.cpp
+++ b/Engine/Source/Developer/ToolMenus/Private/ToolMenus.cpp
@@ -2061,17 +2061,7 @@ TSharedRef<SWidget> UToolMenus::GenerateWidget(UToolMenu* GeneratedMenu)
 {
 	CleanupStaleWidgetsNextTick();
 
-	FGeneratedToolMenuWidgets& WidgetsForMenuName = GeneratedMenuWidgets.FindOrAdd(GeneratedMenu->MenuName);
-
-	// Store a copy so that we can call 'Refresh' on menus not in the database
-	FGeneratedToolMenuWidget& GeneratedMenuWidget = WidgetsForMenuName.Instances.AddDefaulted_GetRef();
-	GeneratedMenuWidget.GeneratedMenu = DuplicateObject<UToolMenu>(GeneratedMenu, this, MakeUniqueObjectName(this, UToolMenus::StaticClass(), FName("MenuForRefresh")));
-	GeneratedMenuWidget.GeneratedMenu->bShouldCleanupContextOnDestroy = true;
-	// Copy native properties that serialize does not
-	GeneratedMenuWidget.GeneratedMenu->Context = GeneratedMenu->Context;
-	GeneratedMenuWidget.GeneratedMenu->StyleSet = GeneratedMenu->StyleSet;
-	GeneratedMenuWidget.GeneratedMenu->StyleName = GeneratedMenu->StyleName;
-
+	TSharedPtr<SWidget> GeneratedWidget;
 	if (GeneratedMenu->IsEditing())
 	{
 		// Convert toolbar into menu during editing
@@ -2101,8 +2091,7 @@ TSharedRef<SWidget> UToolMenus::GenerateWidget(UToolMenu* GeneratedMenu)
 			MenuBuilder.GetMultiBox()->ModifyBlockWidgetAfterMake = GeneratedMenu->ModifyBlockWidgetAfterMake;
 		}
 		TSharedRef<SWidget> Result = MenuBuilder.MakeWidget();
-		GeneratedMenuWidget.Widget = Result;
-		return Result;
+		GeneratedWidget = Result;
 	}
 	else if (GeneratedMenu->MenuType == EMultiBoxType::Menu)
 	{
@@ -2116,8 +2105,7 @@ TSharedRef<SWidget> UToolMenus::GenerateWidget(UToolMenu* GeneratedMenu)
 		MenuBuilder.SetExtendersEnabled(GeneratedMenu->bExtendersEnabled);
 		PopulateMenuBuilder(MenuBuilder, GeneratedMenu);
 		TSharedRef<SWidget> Result = MenuBuilder.MakeWidget(nullptr, GeneratedMenu->MaxHeight);
-		GeneratedMenuWidget.Widget = Result;
-		return Result;
+		GeneratedWidget = Result;
 	}
 	else if (GeneratedMenu->MenuType == EMultiBoxType::MenuBar)
 	{
@@ -2131,8 +2119,7 @@ TSharedRef<SWidget> UToolMenus::GenerateWidget(UToolMenu* GeneratedMenu)
 		MenuBarBuilder.SetExtendersEnabled(GeneratedMenu->bExtendersEnabled);
 		PopulateMenuBarBuilder(MenuBarBuilder, GeneratedMenu);
 		TSharedRef<SWidget> Result = MenuBarBuilder.MakeWidget();
-		GeneratedMenuWidget.Widget = Result;
-		return Result;
+		GeneratedWidget = Result;
 	}
 	else if (GeneratedMenu->MenuType == EMultiBoxType::ToolBar || GeneratedMenu->MenuType == EMultiBoxType::VerticalToolBar || GeneratedMenu->MenuType == EMultiBoxType::UniformToolBar || GeneratedMenu->MenuType == EMultiBoxType::SlimHorizontalToolBar)
 	{
@@ -2147,11 +2134,29 @@ TSharedRef<SWidget> UToolMenus::GenerateWidget(UToolMenu* GeneratedMenu)
 
 		PopulateToolBarBuilder(ToolbarBuilder, GeneratedMenu);
 		TSharedRef<SWidget> Result = ToolbarBuilder.MakeWidget();
-		GeneratedMenuWidget.Widget = Result;
-		return Result;
+		GeneratedWidget = Result;
 	}
 
-	return SNullWidget::NullWidget;
+	FGeneratedToolMenuWidgets& WidgetsForMenuName = GeneratedMenuWidgets.FindOrAdd(GeneratedMenu->MenuName);
+
+	// Store a copy so that we can call 'Refresh' on menus not in the database
+	FGeneratedToolMenuWidget& GeneratedMenuWidget = WidgetsForMenuName.Instances.AddDefaulted_GetRef();
+	GeneratedMenuWidget.GeneratedMenu = DuplicateObject<UToolMenu>(GeneratedMenu, this, MakeUniqueObjectName(this, UToolMenus::StaticClass(), FName("MenuForRefresh")));
+	GeneratedMenuWidget.GeneratedMenu->bShouldCleanupContextOnDestroy = true;
+	// Copy native properties that serialize does not
+	GeneratedMenuWidget.GeneratedMenu->Context = GeneratedMenu->Context;
+	GeneratedMenuWidget.GeneratedMenu->StyleSet = GeneratedMenu->StyleSet;
+	GeneratedMenuWidget.GeneratedMenu->StyleName = GeneratedMenu->StyleName;
+
+	if (GeneratedWidget)
+	{
+		GeneratedMenuWidget.Widget = GeneratedWidget;
+		return GeneratedWidget.ToSharedRef();
+	}
+	else
+	{
+		return SNullWidget::NullWidget;
+	}
 }
 
 void UToolMenus::ModifyEntryForEditDialog(FToolMenuEntry& Entry)
diff --git a/Engine/Source/Programs/ToolMenusTests/Tests/ToolMenus/ToolMenusTest.cpp b/Engine/Source/Programs/ToolMenusTests/Tests/ToolMenus/ToolMenusTest.cpp
--- a/Engine/Source/Programs/ToolMenusTests/Tests/ToolMenus/ToolMenusTest.cpp
+++ b/Engine/Source/Programs/ToolMenusTests/Tests/ToolMenus/ToolMenusTest.cpp
@@ -57,4 +57,24 @@ TEST_CASE("Developer::ToolMenus::GenerateWidget calls dynamic section lambdas",
 	ToolMenus->GenerateWidget("MyMenu", FToolMenuContext());
 
 	CHECK(bWasLambdaCalled);
-}
\ No newline at end of file
+}
+
+// Repro of UE-201151.
+TEST_CASE("Developer::ToolMenus::GenerateWidget can handle simultaneous AddReferencedObjects calls in legacy dynamic sections", "[ToolMenus]")
+{
+	UToolMenus* ToolMenus = CreateUniqueUToolMenusInstance();
+
+	UToolMenu* ToolMenu = ToolMenus->RegisterMenu("MyMenu");
+
+	REQUIRE(ToolMenu);
+
+	// This simulates the crash of UE-201151 that occurred in UToolMenus::GenerateWidget(UToolMenu*) after a complex delegate
+	// triggered a call to UToolMenus::AddReferencedObjects while the delegate was still executing.
+	ToolMenu->AddDynamicSection("MyDynamicLegacySection", FNewToolMenuDelegateLegacy::CreateLambda([ToolMenus](FMenuBuilder&, UToolMenu*) {
+		TArray<UObject*> Array;
+		FReferenceFinder Finder(Array);
+		UToolMenus::AddReferencedObjects(ToolMenus, Finder);
+	}));
+
+	ToolMenus->GenerateWidget("MyMenu", FToolMenuContext());
+}
