ToolMenus: Fix a bug where menu contexts would be incorrectly cleaned up during menu refresh
This fixes a crash on hovering over the "transform" submenu in the light mixer context menu

diff --git a/Engine/Source/Developer/ToolMenus/Private/ToolMenus.cpp b/Engine/Source/Developer/ToolMenus/Private/ToolMenus.cpp
--- a/Engine/Source/Developer/ToolMenus/Private/ToolMenus.cpp
+++ b/Engine/Source/Developer/ToolMenus/Private/ToolMenus.cpp
@@ -2141,6 +2141,7 @@ TSharedRef<SWidget> UToolMenus::GenerateWidget(UToolMenu* GeneratedMenu)
 
 	// Store a copy so that we can call 'Refresh' on menus not in the database
 	FGeneratedToolMenuWidget& GeneratedMenuWidget = WidgetsForMenuName.Instances.AddDefaulted_GetRef();
+	GeneratedMenuWidget.OriginalMenu = GeneratedMenu;
 	GeneratedMenuWidget.GeneratedMenu = DuplicateObject<UToolMenu>(GeneratedMenu, this, MakeUniqueObjectName(this, UToolMenus::StaticClass(), FName("MenuForRefresh")));
 	GeneratedMenuWidget.GeneratedMenu->bShouldCleanupContextOnDestroy = true;
 	// Copy native properties that serialize does not
@@ -2302,6 +2303,13 @@ bool UToolMenus::RefreshMenuWidget(const FName InName, FGeneratedToolMenuWidget&
 
 	// Regenerate menu from database
 	GeneratedMenuWidget.GeneratedMenu->bShouldCleanupContextOnDestroy = false; // The new menu will do this
+
+	// GeneratedMenuWidget.GeneratedMenu is a copy of the original menu, so we also need to make sure the original menu does not clean up its context
+	if(UToolMenu* OriginalMenu = GeneratedMenuWidget.OriginalMenu.Get())
+	{
+		OriginalMenu->bShouldCleanupContextOnDestroy = false;
+	}
+	
 	UToolMenu* GeneratedMenu = GenerateMenu(InName, GeneratedMenuWidget.GeneratedMenu->Context);
 	GeneratedMenuWidget.GeneratedMenu = GeneratedMenu;
 
diff --git a/Engine/Source/Developer/ToolMenus/Public/ToolMenus.h b/Engine/Source/Developer/ToolMenus/Public/ToolMenus.h
--- a/Engine/Source/Developer/ToolMenus/Public/ToolMenus.h
+++ b/Engine/Source/Developer/ToolMenus/Public/ToolMenus.h
@@ -34,8 +34,14 @@ struct FToolMenuSection;
 
 struct FGeneratedToolMenuWidget
 {
+	// A copy of the menu so we can refresh menus not in the database
 	TObjectPtr<UToolMenu> GeneratedMenu;
+
+	// The actual widget for the menu
 	TWeakPtr<SWidget> Widget;
+
+	// Weak ptr to the original menu that owns the widget
+	TWeakObjectPtr<UToolMenu> OriginalMenu;
 };
 
 struct FGeneratedToolMenuWidgets
