Log enum name instead of int value in ToolMenus warning

Also update the text to match another warning message. This warning fires when a menu-type ToolMenu iterates its entries to populate a menu builder but finds an entry that menus don't support, such as a toolbar button.

diff --git a/Engine/Source/Developer/ToolMenus/Private/ToolMenus.cpp b/Engine/Source/Developer/ToolMenus/Private/ToolMenus.cpp
--- a/Engine/Source/Developer/ToolMenus/Private/ToolMenus.cpp
+++ b/Engine/Source/Developer/ToolMenus/Private/ToolMenus.cpp
@@ -356,7 +356,7 @@ public:
 		}
 		else
 		{
-			UE_LOG(LogToolMenus, Warning, TEXT("Menu '%s', item '%s', type not currently supported: %d"), *MenuData->MenuName.ToString(), *BlockNameOverride.ToString(), int(Block.Type));
+			UE_LOG(LogToolMenus, Warning, TEXT("Menu '%s', item '%s', Menus do not support: %s"), *MenuData->MenuName.ToString(), *BlockNameOverride.ToString(), *UEnum::GetValueAsString(Block.Type));
 		}
 	};
 
