ToolMenus: Added reusable UToolMenuWidgetCollectionContext class. Allows adding/finding widgets in a generic way, which should reduce the need to create custom context classes.

diff --git a/Engine/Source/Developer/ToolMenus/Private/ToolMenuWidgetCollectionContext.cpp b/Engine/Source/Developer/ToolMenus/Private/ToolMenuWidgetCollectionContext.cpp
new file mode 100644
--- /dev/null
+++ b/Engine/Source/Developer/ToolMenus/Private/ToolMenuWidgetCollectionContext.cpp
@@ -0,0 +1,60 @@
+// Copyright Epic Games, Inc. All Rights Reserved.
+
+#include "ToolMenuWidgetCollectionContext.h"
+
+#include "Widgets/SWidget.h"
+
+#include UE_INLINE_GENERATED_CPP_BY_NAME(ToolMenuWidgetCollectionContext)
+
+UToolMenuWidgetCollectionContext* UToolMenuWidgetCollectionContext::Get(
+	FToolMenuContext& Context,
+	bool bCreateIfNeeded /*= true*/)
+{
+	if (UToolMenuWidgetCollectionContext* ExistingInstance = Context.FindContext<UToolMenuWidgetCollectionContext>())
+	{
+		return ExistingInstance;
+	}
+
+	if (bCreateIfNeeded)
+	{
+		UToolMenuWidgetCollectionContext* NewInstance = NewObject<UToolMenuWidgetCollectionContext>();
+		Context.AddObject(NewInstance);
+		return NewInstance;
+	}
+
+	return nullptr;
+}
+
+void UToolMenuWidgetCollectionContext::AddWidget(const TSharedRef<SWidget>& Widget)
+{
+	Widgets.AddUnique(Widget);
+}
+
+void UToolMenuWidgetCollectionContext::EnumerateWidgets(const TFunctionRef<bool(const TSharedPtr<SWidget>& Widget)>& Callback) const
+{
+	for (const TWeakPtr<SWidget>& Widget : Widgets)
+	{
+		if (TSharedPtr<SWidget> PinnedWidget = Widget.Pin())
+		{
+			if (!Callback(PinnedWidget))
+			{
+				break;
+			}
+		}
+	}
+}
+
+TSharedPtr<SWidget> UToolMenuWidgetCollectionContext::FindWidgetByClassType(const FName& WidgetClassType) const
+{
+	TSharedPtr<SWidget> FoundWidget;
+	EnumerateWidgets([&FoundWidget, &WidgetClassType](const TSharedPtr<SWidget>& Widget)
+	{
+		if (Widget->GetWidgetClass().GetWidgetType() == WidgetClassType)
+		{
+			FoundWidget = Widget;
+			return false;
+		}
+		return true;
+	});
+	return FoundWidget;
+}
diff --git a/Engine/Source/Developer/ToolMenus/Public/ToolMenuWidgetCollectionContext.h b/Engine/Source/Developer/ToolMenus/Public/ToolMenuWidgetCollectionContext.h
new file mode 100644
--- /dev/null
+++ b/Engine/Source/Developer/ToolMenus/Public/ToolMenuWidgetCollectionContext.h
@@ -0,0 +1,61 @@
+// Copyright Epic Games, Inc. All Rights Reserved.
+
+#pragma once
+
+#include "ToolMenuContext.h"
+
+#include "ToolMenuWidgetCollectionContext.generated.h"
+
+class SWidget;
+
+/**
+ * Context class containing a collection of widgets that can be added for later retrieval.
+ * @see FToolMenuContext
+ */
+UCLASS(MinimalAPI)
+class UToolMenuWidgetCollectionContext : public UToolMenuContextBase
+{
+	GENERATED_BODY()
+
+public:
+	/**
+	 * Get the widget collection context object from the parameter context.
+	 * @param Context The tool menu context to retrieve the context object from.
+	 * @param bCreateIfNeeded If true (default) then the context object is created if one wasn't found.
+	 * @return The context object, or null if one wasn't found and bCreateIfNeeded was false.
+	 */
+	TOOLMENUS_API [[nodiscard]] static UToolMenuWidgetCollectionContext* Get(FToolMenuContext& Context, bool bCreateIfNeeded = true);
+
+	/**
+	 * Add a widget to this context for later retrieval.
+	 * @param Widget The widget to add to the context.
+	 * @see FindContextWidget
+	 */
+	TOOLMENUS_API void AddWidget(const TSharedRef<SWidget>& Widget);
+
+	/**
+	 * Find a previously added widget. Uses the widget's type via WidgetType::StaticWidgetClass.
+	 * @tparam WidgetType The type of widget to find.
+	 * @return The widget, if found. A null pointer otherwise.
+	 * @note This requires that the widget implements SLATE_DECLARE_WIDGET in its class definition.
+	 * @note This function performs an exact match on the widget type - there is no inheritance checking.
+	 * @see AddWidget
+	 */
+	template <typename WidgetType>
+	[[nodiscard]] TSharedPtr<WidgetType> FindWidget() const
+	{
+		const FName TypeToFind = WidgetType::StaticWidgetClass().GetWidgetType();
+		return StaticCastSharedPtr<WidgetType>(FindWidgetByClassType(TypeToFind));
+	}
+
+	/**
+	 * Enumerate all the previously added widgets.
+	 * @param InCallback Callback invoked for each wdget; return true to continue enumeration, or false to stop.
+	 */
+	TOOLMENUS_API void EnumerateWidgets(const TFunctionRef<bool(const TSharedPtr<SWidget>& Widget)>& Callback) const;
+
+private:
+	TOOLMENUS_API [[nodiscard]] TSharedPtr<SWidget> FindWidgetByClassType(const FName& WidgetClassType) const;
+
+	TArray<TWeakPtr<SWidget>> Widgets;
+};
