[ToolMenus] Build fic - move nodiscard before any export macro.

diff --git a/Engine/Source/Developer/ToolMenus/Public/ToolMenuWidgetCollectionContext.h b/Engine/Source/Developer/ToolMenus/Public/ToolMenuWidgetCollectionContext.h
--- a/Engine/Source/Developer/ToolMenus/Public/ToolMenuWidgetCollectionContext.h
+++ b/Engine/Source/Developer/ToolMenus/Public/ToolMenuWidgetCollectionContext.h
@@ -24,7 +24,7 @@ public:
 	 * @param bCreateIfNeeded If true (default) then the context object is created if one wasn't found.
 	 * @return The context object, or null if one wasn't found and bCreateIfNeeded was false.
 	 */
-	TOOLMENUS_API [[nodiscard]] static UToolMenuWidgetCollectionContext* Get(FToolMenuContext& Context, bool bCreateIfNeeded = true);
+	[[nodiscard]] TOOLMENUS_API static UToolMenuWidgetCollectionContext* Get(FToolMenuContext& Context, bool bCreateIfNeeded = true);
 
 	/**
 	 * Add a widget to this context for later retrieval.
@@ -55,7 +55,7 @@ public:
 	TOOLMENUS_API void EnumerateWidgets(const TFunctionRef<bool(const TSharedPtr<SWidget>& Widget)>& Callback) const;
 
 private:
-	TOOLMENUS_API [[nodiscard]] TSharedPtr<SWidget> FindWidgetByClassType(const FName& WidgetClassType) const;
+	[[nodiscard]] TOOLMENUS_API TSharedPtr<SWidget> FindWidgetByClassType(const FName& WidgetClassType) const;
 
 	TArray<TWeakPtr<SWidget>> Widgets;
 };
