Add the concept of section menus to ToolMenu sections

A section menu is currently only rendered for toolbar sections. The section menu can be opened via a new combo button that is added as the left-most entry of toolbar sections. When that combo button is clicked, a menu opens containing all entries of the section.

Section menus are disabled by default, but can be enabled on a section by calling the new FToolMenuSection::SetShowSectionMenu(bool) API.

diff --git a/Engine/Source/Developer/ToolMenus/Private/ToolMenuSection.cpp b/Engine/Source/Developer/ToolMenus/Private/ToolMenuSection.cpp
--- a/Engine/Source/Developer/ToolMenus/Private/ToolMenuSection.cpp
+++ b/Engine/Source/Developer/ToolMenus/Private/ToolMenuSection.cpp
@@ -12,15 +12,14 @@
 
 #include UE_INLINE_GENERATED_CPP_BY_NAME(ToolMenuSection)
 
-FToolMenuSection::FToolMenuSection() :
-	ToolMenuSectionDynamic(nullptr),
-	bIsRegistering(false),
-	bAddedDuringRegister(false)
+FToolMenuSection::FToolMenuSection()
+	: ToolMenuSectionDynamic(nullptr)
+	, bIsRegistering(false)
+	, bAddedDuringRegister(false)
+	, bShowSectionMenu(false)
 {
-
 }
 
-
 void FToolMenuSection::InitSection(const FName InName, const TAttribute< FText >& InLabel, const FToolMenuInsert InPosition)
 {
 	Name = InName;
@@ -35,6 +34,7 @@ void FToolMenuSection::InitGeneratedSectionCopy(const FToolMenuSection& Source,
 	InsertPosition = Source.InsertPosition;
 	Construct = Source.Construct;
 	Context = InContext;
+	bShowSectionMenu = Source.bShowSectionMenu;
 }
 
 bool FToolMenuSection::IsRegistering() const
@@ -165,6 +165,11 @@ const FToolMenuEntry* FToolMenuSection::FindEntry(const FName InName) const
 	return nullptr;
 }
 
+void FToolMenuSection::SetShowSectionMenu(bool InShow)
+{
+	bShowSectionMenu = InShow;
+}
+
 int32 FToolMenuSection::IndexOfBlock(const FName InName) const
 {
 	for (int32 i=0; i < Blocks.Num(); ++i)
diff --git a/Engine/Source/Developer/ToolMenus/Private/ToolMenus.cpp b/Engine/Source/Developer/ToolMenus/Private/ToolMenus.cpp
--- a/Engine/Source/Developer/ToolMenus/Private/ToolMenus.cpp
+++ b/Engine/Source/Developer/ToolMenus/Private/ToolMenus.cpp
@@ -10,6 +10,7 @@
 #include "Framework/MultiBox/MultiBoxBuilder.h"
 #include "Framework/Application/SlateApplication.h"
 #include "Internationalization/Internationalization.h"
+#include "UObject/WeakObjectPtrTemplates.h"
 
 #include "HAL/PlatformApplicationMisc.h" // For clipboard
 #include "Widgets/Layout/SScrollBox.h"
@@ -1444,6 +1445,52 @@ void UToolMenus::PopulateToolBarBuilder(FToolBarBuilder& ToolBarBuilder, UToolMe
 
 		ToolBarBuilder.BeginSection(Section.Name);
 
+		if (Section.bShowSectionMenu)
+		{
+			TWeakObjectPtr<UToolMenu> WeakMenuPtr = MenuData;
+			const FOnGetContent MenuGenerator = FOnGetContent::CreateLambda(
+				[this, WeakMenuPtr, SectionName = Section.Name]() -> TSharedRef<SWidget> {
+					TStrongObjectPtr<UToolMenu> Menu = WeakMenuPtr.Pin();
+					// If we got no menu, there's nothing to do.
+					if (!Menu)
+					{
+						return SNullWidget::NullWidget;
+					}
+
+					FToolMenuSection* const Section = Menu->FindSection(SectionName);
+					// If our section doesn't exist anymore, we have nothing to do.
+					if (!Section)
+					{
+						return SNullWidget::NullWidget;
+					}
+
+					// Most of the following code came from UToolMenus::PopulateMenuBuilder.
+
+					FMenuBuilder MenuBuilder(true, nullptr);
+					MenuBuilder.SetSearchable(Menu->bSearchable);
+
+					// Don't pass in a label here because the label is already displayed as the combo button text
+					// and the menu will only have this single section.
+					MenuBuilder.BeginSection(Section->Name);
+
+					for (FToolMenuEntry& Block : Section->Blocks)
+					{
+						FPopulateMenuBuilderWithToolMenuEntry PopulateMenuBuilderWithToolMenuEntry(
+							MenuBuilder, Menu.Get(), *Section, Block, /* bAllowSubMenuCollapse= */ true);
+						PopulateMenuBuilderWithToolMenuEntry.Populate();
+					}
+
+					MenuBuilder.EndSection();
+
+					AddReferencedContextObjects(MenuBuilder.GetMultiBox(), Menu.Get());
+
+					return MenuBuilder.MakeWidget();
+				});
+
+			ToolBarBuilder.AddComboButton(FUIAction(), MenuGenerator, Section.Label,
+				LOCTEXT("ToolbarSectionMenuTooltip", "Menu that contains all entries of this section"));
+		}
+
 		for (FToolMenuEntry& Block : Section.Blocks)
 		{
 			if (Block.ToolBarData.ConstructLegacy.IsBound())
diff --git a/Engine/Source/Developer/ToolMenus/Public/ToolMenuSection.h b/Engine/Source/Developer/ToolMenus/Public/ToolMenuSection.h
--- a/Engine/Source/Developer/ToolMenus/Public/ToolMenuSection.h
+++ b/Engine/Source/Developer/ToolMenus/Public/ToolMenuSection.h
@@ -85,6 +85,19 @@ public:
 	FToolMenuEntry* FindEntry(const FName InName);
 	const FToolMenuEntry* FindEntry(const FName InName) const;
 
+	/**
+	 * Adds a menu to the section that contain all entries of the section.
+	 *
+	 * The use case of this section menu is to promote the most important entries to show in the toolbar section, while
+	 * the rest of the section's entries are available in the section menu (including the ones promoted to the toolbar).
+	 *
+	 * This feature only applies to sections added to toolbars, not sections added to menus.
+	 * If SetShowSectionMenu(true) is called on a section in a menu, nothing will happen.
+	 *
+	 * @param InShow Shows a section menu if true and doesn't if false (default).
+	 */
+	void SetShowSectionMenu(bool InShow);
+
 private:
 
 	void InitGeneratedSectionCopy(const FToolMenuSection& Source, FToolMenuContext& InContext);
@@ -131,4 +144,5 @@ private:
 
 	bool bIsRegistering;
 	bool bAddedDuringRegister;
+	bool bShowSectionMenu;
 };
