Treat ToolMenu menu entries as toolbar buttons and combo buttons (in the submenu case) to render them in ToolMenu toolbars

diff --git a/Engine/Source/Developer/ToolMenus/Private/ToolMenus.cpp b/Engine/Source/Developer/ToolMenus/Private/ToolMenus.cpp
--- a/Engine/Source/Developer/ToolMenus/Private/ToolMenus.cpp
+++ b/Engine/Source/Developer/ToolMenus/Private/ToolMenus.cpp
@@ -1520,7 +1520,8 @@ void UToolMenus::PopulateToolBarBuilder(FToolBarBuilder& ToolBarBuilder, UToolMe
 			}
 			PRAGMA_ENABLE_DEPRECATION_WARNINGS
 
-			if (Block.Type == EMultiBlockType::ToolBarButton)
+			if (Block.Type == EMultiBlockType::ToolBarButton
+				|| (Block.Type == EMultiBlockType::MenuEntry && !Block.IsSubMenu()))
 			{
 				if (Block.Command.IsValid() && !Block.IsCommandKeybindOnly())
 				{
@@ -1560,7 +1561,8 @@ void UToolMenus::PopulateToolBarBuilder(FToolBarBuilder& ToolBarBuilder, UToolMe
 					ToolBarBuilder.AddComboButton(Block.ToolBarData.OptionsDropdownData->Action, OnGetContent, Block.Label, Block.ToolBarData.OptionsDropdownData->ToolTip, Block.Icon, true, Block.TutorialHighlightName);
 				}
 			}
-			else if (Block.Type == EMultiBlockType::ToolBarComboButton)
+			else if (Block.Type == EMultiBlockType::ToolBarComboButton
+					 || (Block.Type == EMultiBlockType::MenuEntry && Block.IsSubMenu()))
 			{
 				FOnGetContent OnGetContent = ConvertWidgetChoice(Block.ToolBarData.ComboButtonContextMenuGenerator, MenuData->Context);
 				if (OnGetContent.IsBound())
