Add new flag to ToolMenu entries indicating whether they should show up in a toolbar section or just in its section submenu

This flag has no effect if the toolmenu section does not have a section menu.

Note that this change will not display recursive children of the section in the toolbar section even if they are flagged to show in the top level. Only direct child entries of a section with a section menu can be displayed in the top level for now.

diff --git a/Engine/Source/Developer/ToolMenus/Private/ToolMenuEntry.cpp b/Engine/Source/Developer/ToolMenus/Private/ToolMenuEntry.cpp
--- a/Engine/Source/Developer/ToolMenus/Private/ToolMenuEntry.cpp
+++ b/Engine/Source/Developer/ToolMenus/Private/ToolMenuEntry.cpp
@@ -14,27 +14,29 @@
 #include UE_INLINE_GENERATED_CPP_BY_NAME(ToolMenuEntry)
 
 PRAGMA_DISABLE_DEPRECATION_WARNINGS
-FToolMenuEntry::FToolMenuEntry() :
-	Type(EMultiBlockType::None),
-	UserInterfaceActionType(EUserInterfaceActionType::Button),
-	bShouldCloseWindowAfterMenuSelection(true),
-	ScriptObject(nullptr),
-	StyleNameOverride(NAME_None),
-	bAddedDuringRegister(false),
-	bCommandIsKeybindOnly(false)
+FToolMenuEntry::FToolMenuEntry()
+	: Type(EMultiBlockType::None)
+	, UserInterfaceActionType(EUserInterfaceActionType::Button)
+	, bShouldCloseWindowAfterMenuSelection(true)
+	, ScriptObject(nullptr)
+	, StyleNameOverride(NAME_None)
+	, bAddedDuringRegister(false)
+	, bCommandIsKeybindOnly(false)
+	, bShowInToolbarTopLevel(false)
 {
 }
 
-FToolMenuEntry::FToolMenuEntry(const FToolMenuOwner InOwner, const FName InName, EMultiBlockType InType) :
-	Name(InName),
-	Owner(InOwner),
-	Type(InType),
-	UserInterfaceActionType(EUserInterfaceActionType::Button),
-	bShouldCloseWindowAfterMenuSelection(true),
-	ScriptObject(nullptr),
-	StyleNameOverride(NAME_None),
-	bAddedDuringRegister(false),
-	bCommandIsKeybindOnly(false)
+FToolMenuEntry::FToolMenuEntry(const FToolMenuOwner InOwner, const FName InName, EMultiBlockType InType)
+	: Name(InName)
+	, Owner(InOwner)
+	, Type(InType)
+	, UserInterfaceActionType(EUserInterfaceActionType::Button)
+	, bShouldCloseWindowAfterMenuSelection(true)
+	, ScriptObject(nullptr)
+	, StyleNameOverride(NAME_None)
+	, bAddedDuringRegister(false)
+	, bCommandIsKeybindOnly(false)
+	, bShowInToolbarTopLevel(false)
 {
 }
 
@@ -297,6 +299,11 @@ bool FToolMenuEntry::TryExecuteToolUIAction(const FToolMenuContext& InContext)
 	return bCanExecute;
 }
 
+void FToolMenuEntry::SetShowInToolbarTopLevel(bool InTopLevel)
+{
+	bShowInToolbarTopLevel = InTopLevel;
+}
+
 bool FToolMenuEntry::IsScriptObjectDynamicConstruct() const
 {
 	static const FName ConstructMenuEntryName = GET_FUNCTION_NAME_CHECKED(UToolMenuEntryScript, ConstructMenuEntry);
diff --git a/Engine/Source/Developer/ToolMenus/Private/ToolMenus.cpp b/Engine/Source/Developer/ToolMenus/Private/ToolMenus.cpp
--- a/Engine/Source/Developer/ToolMenus/Private/ToolMenus.cpp
+++ b/Engine/Source/Developer/ToolMenus/Private/ToolMenus.cpp
@@ -1493,6 +1493,11 @@ void UToolMenus::PopulateToolBarBuilder(FToolBarBuilder& ToolBarBuilder, UToolMe
 
 		for (FToolMenuEntry& Block : Section.Blocks)
 		{
+			if (Section.bShowSectionMenu && !Block.bShowInToolbarTopLevel)
+			{
+				continue;
+			}
+
 			if (Block.ToolBarData.ConstructLegacy.IsBound())
 			{
 				Block.ToolBarData.ConstructLegacy.Execute(ToolBarBuilder, MenuData);
diff --git a/Engine/Source/Developer/ToolMenus/Public/ToolMenuEntry.h b/Engine/Source/Developer/ToolMenus/Public/ToolMenuEntry.h
--- a/Engine/Source/Developer/ToolMenus/Public/ToolMenuEntry.h
+++ b/Engine/Source/Developer/ToolMenus/Public/ToolMenuEntry.h
@@ -155,6 +155,22 @@ struct TOOLMENUS_API FToolMenuEntry
 	friend struct FToolMenuSection;
 	friend class UToolMenuEntryScript;
 
+	/**
+	 * Show this menu entry in the top-level toolbar section of a toolbar.
+	 *
+	 * A toolbar that has a section menu only shows entries flagged as top-level in the toolbar section itself while the
+	 * others can only be accessed through the section submenu. Use this method to flag a menu entry to show up in that
+	 * top-level toolbar section.
+	 *
+	 * This flag only affects entries of ToolMenus rendered as toolbars that also have a section menu. For ToolMenus
+	 * rendered as menus or toolbar ToolMenus that lack a section menu, the entry is always placed in the top-level
+	 * section itself.
+	 *
+	 * @param InTopLevel Pass true to set this menu entry to appear in the top-level toolbar section for toolbar
+	 * sections with a section menu. Pass false to only show the entry in the section submenu.
+	 */
+	void SetShowInToolbarTopLevel(bool InTopLevel);
+
 private:
 
 	void SetCommand(const TSharedPtr< const FUICommandInfo >& InCommand, TOptional<FName> InName, const TAttribute<FText>& InLabel, const TAttribute<FText>& InToolTip, const TAttribute<FSlateIcon>& InIcon);
@@ -231,4 +247,6 @@ private:
 
 	UPROPERTY()
 	bool bCommandIsKeybindOnly;
+
+	bool bShowInToolbarTopLevel;
 };
