Make UToolMenus::RefreshMenuWidget refresh all widgets of a menu, not just the first

diff --git a/Engine/Source/Developer/ToolMenus/Private/ToolMenus.cpp b/Engine/Source/Developer/ToolMenus/Private/ToolMenus.cpp
--- a/Engine/Source/Developer/ToolMenus/Private/ToolMenus.cpp
+++ b/Engine/Source/Developer/ToolMenus/Private/ToolMenus.cpp
@@ -1722,7 +1722,7 @@ FUIAction UToolMenus::ConvertUIAction(const FToolMenuEntry& Block, const FToolMe
 	{
 		UIAction = ConvertUIAction(Block.Action, Context);
 	}
-	
+
 	if (!UIAction.ExecuteAction.IsBound() && Block.StringExecuteAction.IsBound())
 	{
 		UIAction.ExecuteAction = Block.StringExecuteAction.ToExecuteAction(Context);
@@ -2332,13 +2332,14 @@ void UToolMenus::CleanupStaleWidgets()
 
 bool UToolMenus::RefreshMenuWidget(const FName InName)
 {
+	bool bRefreshedAnyWidget = false;
 	if (FGeneratedToolMenuWidgets* WidgetsForMenuName = GeneratedMenuWidgets.Find(InName))
 	{
 		for (auto Instance = WidgetsForMenuName->Instances.CreateIterator(); Instance; ++Instance)
 		{
 			if (RefreshMenuWidget(InName, *Instance))
 			{
-				return true;
+				bRefreshedAnyWidget = true;
 			}
 			else
 			{
@@ -2347,7 +2348,7 @@ bool UToolMenus::RefreshMenuWidget(const FName InName)
 		}
 	}
 
-	return false;
+	return bRefreshedAnyWidget;
 }
 
 bool UToolMenus::RefreshMenuWidget(const FName InName, FGeneratedToolMenuWidget& GeneratedMenuWidget)
@@ -2365,7 +2366,7 @@ bool UToolMenus::RefreshMenuWidget(const FName InName, FGeneratedToolMenuWidget&
 	{
 		OriginalMenu->bShouldCleanupContextOnDestroy = false;
 	}
-	
+
 	UToolMenu* GeneratedMenu = GenerateMenu(InName, GeneratedMenuWidget.GeneratedMenu->Context);
 	GeneratedMenuWidget.GeneratedMenu = GeneratedMenu;
 
