Recursively collect ToolMenu section children into the top-level of a toolbar section

This applies if the toolbar section has a section menu. In that case, the toolbar section itself will from now on be populated with recursive children if those children are flagged to show up in the toolbar top level via FToolMenuEntry::SetShowInToolbarTopLevel(bool).

Note: to update the top-level status of a grandchild, you might need to call UToolMenus::RefreshMenuWidget for the new top-level status to be reflected in already generated widgets.

diff --git a/Engine/Source/Developer/ToolMenus/Private/ToolMenus.cpp b/Engine/Source/Developer/ToolMenus/Private/ToolMenus.cpp
--- a/Engine/Source/Developer/ToolMenus/Private/ToolMenus.cpp
+++ b/Engine/Source/Developer/ToolMenus/Private/ToolMenus.cpp
@@ -1411,6 +1411,195 @@ void UToolMenus::PopulateMenuBuilder(FMenuBuilder& MenuBuilder, UToolMenu* MenuD
 	AddReferencedContextObjects(MenuBuilder.GetMultiBox(), MenuData);
 }
 
+void UToolMenus::ExtractChildBlocksFromSubMenu(
+	UToolMenu* ParentMenu, FToolMenuEntry& InBlock, TArray<TPair<UToolMenu*, FToolMenuEntry*>>& SubMenuBlockPairs)
+{
+	const bool bBuiltViaDelegate =
+		ConvertWidgetChoice(InBlock.ToolBarData.ComboButtonContextMenuGenerator, ParentMenu->Context).IsBound();
+
+	const bool bHasChildren = (InBlock.Type == EMultiBlockType::ToolBarComboButton
+								  || (InBlock.Type == EMultiBlockType::MenuEntry && InBlock.IsSubMenu()))
+						   && !bBuiltViaDelegate;
+
+	if (!bHasChildren)
+	{
+		return;
+	}
+
+	UToolMenu* const SubMenu = GenerateSubMenu(ParentMenu, InBlock.Name);
+	if (!SubMenu)
+	{
+		return;
+	}
+
+	// Add the blocks reversed to follow a depth-first iteration.
+	for (int i = SubMenu->Sections.Num() - 1; i >= 0; --i)
+	{
+		FToolMenuSection& Section = SubMenu->Sections[i];
+
+		for (int j = Section.Blocks.Num() - 1; j >= 0; --j)
+		{
+			FToolMenuEntry& Block = Section.Blocks[j];
+			SubMenuBlockPairs.Push(TPair<UToolMenu*, FToolMenuEntry*>(SubMenu, &Block));
+		}
+	}
+}
+
+void UToolMenus::PopulateToolBarBuilderWithGrandchildren(
+	FToolBarBuilder& ToolBarBuilder, UToolMenu* ParentMenu, FToolMenuEntry& InBlock)
+{
+	TArray<TPair<UToolMenu*, FToolMenuEntry*>> SubMenuBlockPairs;
+	// Seed the blocks with the passed-in submenu.
+	ExtractChildBlocksFromSubMenu(ParentMenu, InBlock, SubMenuBlockPairs);
+
+	int32 NumIterations = 0;
+	while (SubMenuBlockPairs.Num() > 0)
+	{
+		const TPair<UToolMenu*, FToolMenuEntry*> SubMenuAndBlock = SubMenuBlockPairs.Pop(EAllowShrinking::No);
+		UToolMenu* const SubMenu = SubMenuAndBlock.Key;
+		FToolMenuEntry* const Block = SubMenuAndBlock.Value;
+
+		// Keep track of how many blocks we've visited to ensure we don't loop indefinitely.
+		if (++NumIterations > 5000)
+		{
+			UE_LOG(LogToolMenus, Warning,
+				TEXT("Possible infinite loop for menu with section menu. parent menu: %s, menu: %s, block: %s"),
+				*ParentMenu->MenuName.ToString(), *SubMenu->MenuName.ToString(), *Block->Name.ToString());
+			break;
+		}
+
+		// We pass in true for bSectionHasSubMenu here because we're referring to the top-level section that
+		// called us and not the section we're currently iterating through. It is that top-level section
+		// that we're populating here, not the section the current Block came from.
+		constexpr bool bSectionHasSubMenu = true;
+		PopulateToolBarBuilderWithEntry(ToolBarBuilder, SubMenu, bSectionHasSubMenu, *Block);
+
+		ExtractChildBlocksFromSubMenu(SubMenu, *Block, SubMenuBlockPairs);
+	}
+}
+
+void UToolMenus::PopulateToolBarBuilderWithEntry(
+	FToolBarBuilder& ToolBarBuilder, UToolMenu* MenuData, bool bSectionHasSubMenu, FToolMenuEntry& Block)
+{
+	if (Block.ToolBarData.ConstructLegacy.IsBound())
+	{
+		Block.ToolBarData.ConstructLegacy.Execute(ToolBarBuilder, MenuData);
+		return;
+	}
+
+	const bool bAddEntry = !bSectionHasSubMenu || (bSectionHasSubMenu && Block.bShowInToolbarTopLevel);
+	if (!bAddEntry)
+	{
+		return;
+	}
+
+	ToolBarBuilder.BeginStyleOverride(Block.StyleNameOverride);
+
+	if (Block.Type == EMultiBlockType::ToolBarButton || (Block.Type == EMultiBlockType::MenuEntry && !Block.IsSubMenu()))
+	{
+		if (Block.Command.IsValid() && !Block.IsCommandKeybindOnly())
+		{
+			bool bPopCommandList = false;
+			TSharedPtr<const FUICommandList> CommandListForAction;
+			if (Block.GetActionForCommand(MenuData->Context, CommandListForAction) != nullptr
+				&& CommandListForAction.IsValid())
+			{
+				ToolBarBuilder.PushCommandList(CommandListForAction.ToSharedRef());
+				bPopCommandList = true;
+			}
+			else
+			{
+				UE_LOG(LogToolMenus, Verbose, TEXT("UI command not found for toolbar entry: %s, toolbar: %s"),
+					*Block.Name.ToString(), *MenuData->MenuName.ToString());
+			}
+
+			ToolBarBuilder.AddToolBarButton(
+				Block.Command, Block.Name, Block.Label, Block.ToolTip, Block.Icon, Block.TutorialHighlightName);
+
+			if (bPopCommandList)
+			{
+				ToolBarBuilder.PopCommandList();
+			}
+		}
+		else if (Block.ScriptObject)
+		{
+			UToolMenuEntryScript* ScriptObject = Block.ScriptObject;
+			TAttribute<FSlateIcon> Icon = ScriptObject->CreateIconAttribute(MenuData->Context);
+			FUIAction UIAction = ConvertUIAction(Block, MenuData->Context);
+			ToolBarBuilder.AddToolBarButton(UIAction, ScriptObject->Data.Name,
+				ScriptObject->CreateLabelAttribute(MenuData->Context),
+				ScriptObject->CreateToolTipAttribute(MenuData->Context), Icon, Block.UserInterfaceActionType,
+				Block.TutorialHighlightName);
+		}
+		else
+		{
+			FUIAction UIAction = ConvertUIAction(Block, MenuData->Context);
+			ToolBarBuilder.AddToolBarButton(UIAction, Block.Name, Block.Label, Block.ToolTip, Block.Icon,
+				Block.UserInterfaceActionType, Block.TutorialHighlightName);
+		}
+
+		if (Block.ToolBarData.OptionsDropdownData.IsValid())
+		{
+			FOnGetContent OnGetContent = ConvertWidgetChoice(
+				Block.ToolBarData.OptionsDropdownData->MenuContentGenerator, MenuData->Context);
+			ToolBarBuilder.AddComboButton(Block.ToolBarData.OptionsDropdownData->Action, OnGetContent, Block.Label,
+				Block.ToolBarData.OptionsDropdownData->ToolTip, Block.Icon, true, Block.TutorialHighlightName);
+		}
+	}
+	else if (Block.Type == EMultiBlockType::ToolBarComboButton
+			 || (Block.Type == EMultiBlockType::MenuEntry && Block.IsSubMenu()))
+	{
+		FOnGetContent OnGetContent = ConvertWidgetChoice(
+			Block.ToolBarData.ComboButtonContextMenuGenerator, MenuData->Context);
+		if (OnGetContent.IsBound())
+		{
+			FUIAction UIAction = ConvertUIAction(Block, MenuData->Context);
+			ToolBarBuilder.AddComboButton(UIAction, OnGetContent, Block.Label, Block.ToolTip, Block.Icon,
+				Block.ToolBarData.bSimpleComboBox, Block.TutorialHighlightName);
+		}
+		else
+		{
+			FOnGetContent Delegate = FOnGetContent::CreateUObject(
+				this, &UToolMenus::GenerateToolbarComboButtonMenu, TWeakObjectPtr<UToolMenu>(MenuData), Block.Name);
+			FUIAction UIAction = ConvertUIAction(Block, MenuData->Context);
+			ToolBarBuilder.AddComboButton(UIAction, Delegate, Block.Label, Block.ToolTip, Block.Icon,
+				Block.ToolBarData.bSimpleComboBox, Block.TutorialHighlightName);
+		}
+	}
+	else if (Block.Type == EMultiBlockType::Separator)
+	{
+		ToolBarBuilder.AddSeparator(Block.Name);
+	}
+	else if (Block.Type == EMultiBlockType::Widget)
+	{
+		TSharedPtr<SWidget> Widget;
+
+		if (Block.MakeCustomWidget.IsBound())
+		{
+			FToolMenuCustomWidgetContext EntryWidgetContext;
+			TSharedRef<FMultiBox> MultiBox = ToolBarBuilder.GetMultiBox();
+			EntryWidgetContext.StyleSet = MultiBox->GetStyleSet();
+			EntryWidgetContext.StyleName = MultiBox->GetStyleName();
+			Widget = Block.MakeCustomWidget.Execute(MenuData->Context, EntryWidgetContext);
+		}
+		PRAGMA_DISABLE_DEPRECATION_WARNINGS
+		else if (Block.MakeWidget.IsBound())
+		{
+			Widget = Block.MakeWidget.Execute(MenuData->Context);
+		}
+		PRAGMA_ENABLE_DEPRECATION_WARNINGS
+
+		ToolBarBuilder.AddWidget(Widget.ToSharedRef(), Block.TutorialHighlightName, Block.WidgetData.bSearchable);
+	}
+	else
+	{
+		UE_LOG(LogToolMenus, Warning, TEXT("Toolbar '%s', item '%s', Toolbars do not support: %s"),
+			*MenuData->MenuName.ToString(), *Block.Name.ToString(), *UEnum::GetValueAsString(Block.Type));
+	}
+
+	ToolBarBuilder.EndStyleOverride();
+}
+
 void UToolMenus::PopulateToolBarBuilder(FToolBarBuilder& ToolBarBuilder, UToolMenu* MenuData)
 {
 	if (GetEditMenusMode() && !MenuData->IsEditing() && EditMenuDelegate.IsBound())
@@ -1493,108 +1682,11 @@ void UToolMenus::PopulateToolBarBuilder(FToolBarBuilder& ToolBarBuilder, UToolMe
 
 		for (FToolMenuEntry& Block : Section.Blocks)
 		{
-			if (Section.bShowSectionMenu && !Block.bShowInToolbarTopLevel)
+			PopulateToolBarBuilderWithEntry(ToolBarBuilder, MenuData, Section.bShowSectionMenu, Block);
+			if (Section.bShowSectionMenu)
 			{
-				continue;
+				PopulateToolBarBuilderWithGrandchildren(ToolBarBuilder, MenuData, Block);
 			}
-
-			if (Block.ToolBarData.ConstructLegacy.IsBound())
-			{
-				Block.ToolBarData.ConstructLegacy.Execute(ToolBarBuilder, MenuData);
-				continue;
-			}
-
-			FUIAction UIAction = ConvertUIAction(Block, MenuData->Context);
-
-			ToolBarBuilder.BeginStyleOverride(Block.StyleNameOverride);
-
-			TSharedPtr<SWidget> Widget;
-
-			if (Block.MakeCustomWidget.IsBound())
-			{
-				FToolMenuCustomWidgetContext EntryWidgetContext;
-				TSharedRef<FMultiBox> MultiBox = ToolBarBuilder.GetMultiBox();
-				EntryWidgetContext.StyleSet = MultiBox->GetStyleSet();
-				EntryWidgetContext.StyleName = MultiBox->GetStyleName();
-				Widget = Block.MakeCustomWidget.Execute(MenuData->Context, EntryWidgetContext);
-			}
-			PRAGMA_DISABLE_DEPRECATION_WARNINGS
-			else if(Block.MakeWidget.IsBound())
-			{
-				Widget = Block.MakeWidget.Execute(MenuData->Context);
-			}
-			PRAGMA_ENABLE_DEPRECATION_WARNINGS
-
-			if (Block.Type == EMultiBlockType::ToolBarButton
-				|| (Block.Type == EMultiBlockType::MenuEntry && !Block.IsSubMenu()))
-			{
-				if (Block.Command.IsValid() && !Block.IsCommandKeybindOnly())
-				{
-					bool bPopCommandList = false;
-					TSharedPtr<const FUICommandList> CommandListForAction;
-					if (Block.GetActionForCommand(MenuData->Context, CommandListForAction) != nullptr && CommandListForAction.IsValid())
-					{
-						ToolBarBuilder.PushCommandList(CommandListForAction.ToSharedRef());
-						bPopCommandList = true;
-					}
-					else
-					{
-						UE_LOG(LogToolMenus, Verbose, TEXT("UI command not found for toolbar entry: %s, toolbar: %s"), *Block.Name.ToString(), *MenuData->MenuName.ToString());
-					}
-
-					ToolBarBuilder.AddToolBarButton(Block.Command, Block.Name, Block.Label, Block.ToolTip, Block.Icon, Block.TutorialHighlightName);
-
-					if (bPopCommandList)
-					{
-						ToolBarBuilder.PopCommandList();
-					}
-				}
-				else if (Block.ScriptObject)
-				{
-					UToolMenuEntryScript* ScriptObject = Block.ScriptObject;
-					TAttribute<FSlateIcon> Icon = ScriptObject->CreateIconAttribute(MenuData->Context);
-					ToolBarBuilder.AddToolBarButton(UIAction, ScriptObject->Data.Name, ScriptObject->CreateLabelAttribute(MenuData->Context), ScriptObject->CreateToolTipAttribute(MenuData->Context), Icon, Block.UserInterfaceActionType, Block.TutorialHighlightName);
-				}
-				else
-				{
-					ToolBarBuilder.AddToolBarButton(UIAction, Block.Name, Block.Label, Block.ToolTip, Block.Icon, Block.UserInterfaceActionType, Block.TutorialHighlightName);
-				}
-
-				if (Block.ToolBarData.OptionsDropdownData.IsValid())
-				{
-					FOnGetContent OnGetContent = ConvertWidgetChoice(Block.ToolBarData.OptionsDropdownData->MenuContentGenerator, MenuData->Context);
-					ToolBarBuilder.AddComboButton(Block.ToolBarData.OptionsDropdownData->Action, OnGetContent, Block.Label, Block.ToolBarData.OptionsDropdownData->ToolTip, Block.Icon, true, Block.TutorialHighlightName);
-				}
-			}
-			else if (Block.Type == EMultiBlockType::ToolBarComboButton
-					 || (Block.Type == EMultiBlockType::MenuEntry && Block.IsSubMenu()))
-			{
-				FOnGetContent OnGetContent = ConvertWidgetChoice(Block.ToolBarData.ComboButtonContextMenuGenerator, MenuData->Context);
-				if (OnGetContent.IsBound())
-				{
-					ToolBarBuilder.AddComboButton(UIAction, OnGetContent, Block.Label, Block.ToolTip, Block.Icon, Block.ToolBarData.bSimpleComboBox, Block.TutorialHighlightName);
-				}
-				else
-				{
-					FName SubMenuFullName = JoinMenuPaths(MenuData->MenuName, Block.Name);
-					FOnGetContent Delegate = FOnGetContent::CreateUObject(this, &UToolMenus::GenerateToolbarComboButtonMenu, TWeakObjectPtr<UToolMenu>(MenuData), Block.Name);
-					ToolBarBuilder.AddComboButton(UIAction, Delegate, Block.Label, Block.ToolTip, Block.Icon, Block.ToolBarData.bSimpleComboBox, Block.TutorialHighlightName);
-				}
-			}
-			else if (Block.Type == EMultiBlockType::Separator)
-			{
-				ToolBarBuilder.AddSeparator(Block.Name);
-			}
-			else if (Block.Type == EMultiBlockType::Widget)
-			{
-				ToolBarBuilder.AddWidget(Widget.ToSharedRef(), Block.TutorialHighlightName, Block.WidgetData.bSearchable);
-			}
-			else
-			{
-				UE_LOG(LogToolMenus, Warning, TEXT("Toolbar '%s', item '%s', Toolbars do not support: %s"), *MenuData->MenuName.ToString(), *Block.Name.ToString(), *UEnum::GetValueAsString(Block.Type));
-			}
-
-			ToolBarBuilder.EndStyleOverride();
 		}
 
 		ToolBarBuilder.EndSection();
diff --git a/Engine/Source/Developer/ToolMenus/Public/ToolMenus.h b/Engine/Source/Developer/ToolMenus/Public/ToolMenus.h
--- a/Engine/Source/Developer/ToolMenus/Public/ToolMenus.h
+++ b/Engine/Source/Developer/ToolMenus/Public/ToolMenus.h
@@ -385,6 +385,15 @@ private:
 
 	void PopulateMenuBuilder(FMenuBuilder& MenuBuilder, UToolMenu* MenuData);
 	void PopulateMenuBarBuilder(FMenuBarBuilder& MenuBarBuilder, UToolMenu* MenuData);
+
+	void ExtractChildBlocksFromSubMenu(
+		UToolMenu* ParentMenu, FToolMenuEntry& InBlock, TArray<TPair<UToolMenu*, FToolMenuEntry*>>& SubMenuBlockPairs);
+	void PopulateToolBarBuilderWithGrandchildren(
+		FToolBarBuilder& ToolBarBuilder, UToolMenu* ParentMenu, FToolMenuEntry& InBlock);
+	void PopulateToolBarBuilderWithEntry(FToolBarBuilder& PopulateToolBarBuilderWithEntry,
+		UToolMenu* MenuData,
+		bool bSectionHasSubMenu,
+		FToolMenuEntry& Entry);
 	void PopulateToolBarBuilder(FToolBarBuilder& ToolBarBuilder, UToolMenu* MenuData);
 
 	TSharedRef<SWidget> GenerateToolbarComboButtonMenu(TWeakObjectPtr<UToolMenu> InParent, const FName InBlockName);
