Add Default, First, Middle, and Last alignment to ToolMenu sections

The new alignment support layers on top of the existing insertion-order support. Sections are added to the toolbar in groups where we first add all Left-aligned sections according to their insertion order, then all First-aligned sections, and then the Middle and Last-aligned sections. Note that this can break the specified insertion order if the user mixed alignments, but otherwise it's maintained.

The First and Default alignments are grouped together, with First-aligned sections appearing before Default-aligned sections. The result is a toolbar with visual empty space between the First-Default group, Middle group, and the Last group. A caveat here is that the Middle group won't necessarily appear perfectly centered on the toolbar because it's position depends on the sizes of the First-Default and Last groups.

diff --git a/Engine/Source/Developer/ToolMenus/Private/ToolMenuSection.cpp b/Engine/Source/Developer/ToolMenus/Private/ToolMenuSection.cpp
--- a/Engine/Source/Developer/ToolMenus/Private/ToolMenuSection.cpp
+++ b/Engine/Source/Developer/ToolMenus/Private/ToolMenuSection.cpp
@@ -14,6 +14,7 @@
 
 FToolMenuSection::FToolMenuSection()
 	: ToolMenuSectionDynamic(nullptr)
+	, Alignment(EToolMenuSectionAlign::Default)
 	, bIsRegistering(false)
 	, bAddedDuringRegister(false)
 {
@@ -33,6 +34,7 @@ void FToolMenuSection::InitGeneratedSectionCopy(const FToolMenuSection& Source,
 	InsertPosition = Source.InsertPosition;
 	Construct = Source.Construct;
 	Context = InContext;
+	Alignment = Source.Alignment;
 }
 
 bool FToolMenuSection::IsRegistering() const
diff --git a/Engine/Source/Developer/ToolMenus/Private/ToolMenus.cpp b/Engine/Source/Developer/ToolMenus/Private/ToolMenus.cpp
--- a/Engine/Source/Developer/ToolMenus/Private/ToolMenus.cpp
+++ b/Engine/Source/Developer/ToolMenus/Private/ToolMenus.cpp
@@ -14,6 +14,7 @@
 
 #include "HAL/PlatformApplicationMisc.h" // For clipboard
 #include "Widgets/Layout/SScrollBox.h"
+#include "Widgets/Layout/SSpacer.h"
 
 #include UE_INLINE_GENERATED_CPP_BY_NAME(ToolMenus)
 
@@ -1619,22 +1620,56 @@ void UToolMenus::PopulateToolBarBuilder(FToolBarBuilder& ToolBarBuilder, UToolMe
 		ToolBarBuilder.EndSection();
 	}
 
-	for (FToolMenuSection& Section : MenuData->Sections)
+	// Add the sections grouped by alignment with SSpacers in between. This visually separates them and allows users
+	// to align sections to appear first, middle, or last. Default-aligned sections appear grouped with first-aligned
+	// sections but appear after them.
+	const TArray<EToolMenuSectionAlign> SectionAlignments = { EToolMenuSectionAlign::First,
+															  EToolMenuSectionAlign::Default,
+															  EToolMenuSectionAlign::Middle,
+															  EToolMenuSectionAlign::Last };
+	bool bDidAddSection = false;
+	for (const EToolMenuSectionAlign CurrentAlignment : SectionAlignments)
 	{
-		if (Section.Construct.NewToolBarDelegateLegacy.IsBound())
+		const bool bIsMiddleOrLast = CurrentAlignment == EToolMenuSectionAlign::Middle
+								  || CurrentAlignment == EToolMenuSectionAlign::Last;
+
+		// Add a spacer before the middle and last alignment groups, and only if we've already added a section to a
+		// previous alignment group.
+		if (bIsMiddleOrLast && bDidAddSection)
 		{
-			Section.Construct.NewToolBarDelegateLegacy.Execute(ToolBarBuilder, MenuData);
-			continue;
+			ToolBarBuilder.AddWidget(SNew(SSpacer), NAME_None, true, HAlign_Right);
 		}
 
-		ToolBarBuilder.BeginSection(Section.Name);
-
-		for (FToolMenuEntry& Block : Section.Blocks)
+		// Keep track if this is the first section we're adding for the CurrentAlignment. Make an exception if the
+		// current alignment is Default and we already added a first-aligned section, because in that case this
+		// isn't the first section in the "group" since we're displaying first and default-aligned sections together.
+		bool bFirstSectionInAlignmentGroup = CurrentAlignment == EToolMenuSectionAlign::Default ? !bDidAddSection : true;
+		for (FToolMenuSection& Section : MenuData->Sections)
 		{
-			PopulateToolBarBuilderWithEntry(ToolBarBuilder, MenuData, Block);
-		}
+			if (Section.Alignment != CurrentAlignment)
+			{
+				continue;
+			}
 
-		ToolBarBuilder.EndSection();
+			if (Section.Construct.NewToolBarDelegateLegacy.IsBound())
+			{
+				Section.Construct.NewToolBarDelegateLegacy.Execute(ToolBarBuilder, MenuData);
+				continue;
+			}
+
+			const bool bSectionShouldHaveSeparator = !bFirstSectionInAlignmentGroup;
+			ToolBarBuilder.BeginSection(Section.Name, bSectionShouldHaveSeparator);
+
+			for (FToolMenuEntry& Block : Section.Blocks)
+			{
+				PopulateToolBarBuilderWithEntry(ToolBarBuilder, MenuData, Block);
+			}
+
+			ToolBarBuilder.EndSection();
+
+			bDidAddSection = true;
+			bFirstSectionInAlignmentGroup = false;
+		}
 	}
 
 	AddReferencedContextObjects(ToolBarBuilder.GetMultiBox(), MenuData);
diff --git a/Engine/Source/Developer/ToolMenus/Public/ToolMenuSection.h b/Engine/Source/Developer/ToolMenus/Public/ToolMenuSection.h
--- a/Engine/Source/Developer/ToolMenus/Public/ToolMenuSection.h
+++ b/Engine/Source/Developer/ToolMenus/Public/ToolMenuSection.h
@@ -21,6 +21,15 @@ public:
 	void ConstructSections(UToolMenu* Menu, const FToolMenuContext& Context);
 };
 
+UENUM(BlueprintType)
+enum class EToolMenuSectionAlign : uint8
+{
+	Default,
+	First,
+	Middle,
+	Last,
+};
+
 USTRUCT(BlueprintType)
 struct TOOLMENUS_API FToolMenuSection
 {
@@ -127,6 +136,8 @@ public:
 
 	FNewSectionConstructChoice Construct;
 
+	EToolMenuSectionAlign Alignment;
+
 private:
 
 	bool bIsRegistering;
diff --git a/Engine/Source/Editor/LevelEditor/Private/SLevelViewport.cpp b/Engine/Source/Editor/LevelEditor/Private/SLevelViewport.cpp
--- a/Engine/Source/Editor/LevelEditor/Private/SLevelViewport.cpp
+++ b/Engine/Source/Editor/LevelEditor/Private/SLevelViewport.cpp
@@ -1933,6 +1933,7 @@ TSharedPtr<SWidget> SLevelViewport::MakeViewportToolbar()
 
 			{
 				FToolMenuSection& RightSection = ViewportToolbarMenu->FindOrAddSection("Right");
+				RightSection.Alignment = EToolMenuSectionAlign::Last;
 
 				{
 					// Stay backward-compatible with the old viewport toolbar.
diff --git a/Engine/Source/Runtime/Slate/Private/Framework/MultiBox/MultiBoxBuilder.cpp b/Engine/Source/Runtime/Slate/Private/Framework/MultiBox/MultiBoxBuilder.cpp
--- a/Engine/Source/Runtime/Slate/Private/Framework/MultiBox/MultiBoxBuilder.cpp
+++ b/Engine/Source/Runtime/Slate/Private/Framework/MultiBox/MultiBoxBuilder.cpp
@@ -646,7 +646,7 @@ void FToolBarBuilder::AddSeparator(FName InExtensionHook)
 	ApplyHook(InExtensionHook, EExtensionHook::After);
 }
 
-void FToolBarBuilder::BeginSection( FName InExtensionHook )
+void FToolBarBuilder::BeginSection(FName InExtensionHook, bool bInSectionShouldHaveSeparator)
 {
 	checkf(CurrentSectionExtensionHook == NAME_None && !bSectionNeedsToBeApplied, TEXT("Did you forget to call EndSection()?"));
 
@@ -655,6 +655,7 @@ void FToolBarBuilder::BeginSection( FName InExtensionHook )
 	// Do not actually apply the section header, because if this section is ended immediately
 	// then nothing ever gets created, preventing empty sections from ever appearing
 	bSectionNeedsToBeApplied = true;
+	bSectionShouldHaveSeparator = bInSectionShouldHaveSeparator;
 	CurrentSectionExtensionHook = InExtensionHook;
 	
 	// Do apply the section beginning if we are in developer "show me all the hooks" mode
@@ -692,12 +693,14 @@ void FToolBarBuilder::ApplySectionBeginning()
 {
 	if (bSectionNeedsToBeApplied)
 	{
-		if( MultiBox->GetBlocks().Num() > 0 || FMultiBoxSettings::DisplayMultiboxHooks.Get() )
+		if (bSectionShouldHaveSeparator
+			&& (MultiBox->GetBlocks().Num() > 0 || FMultiBoxSettings::DisplayMultiboxHooks.Get()))
 		{
 			TSharedRef<FToolBarSeparatorBlock> NewSeparatorBlock = MakeShared<FToolBarSeparatorBlock>(CurrentSectionExtensionHook);
 			NewSeparatorBlock->SetStyleNameOverride(CurrentStyleOverride);
 
 			MultiBox->AddMultiBlock(NewSeparatorBlock);
+			bSectionShouldHaveSeparator = false;
 		}
 		bSectionNeedsToBeApplied = false;
 	}
diff --git a/Engine/Source/Runtime/Slate/Public/Framework/MultiBox/MultiBoxBuilder.h b/Engine/Source/Runtime/Slate/Public/Framework/MultiBox/MultiBoxBuilder.h
--- a/Engine/Source/Runtime/Slate/Public/Framework/MultiBox/MultiBoxBuilder.h
+++ b/Engine/Source/Runtime/Slate/Public/Framework/MultiBox/MultiBoxBuilder.h
@@ -508,6 +508,7 @@ public:
 	FToolBarBuilder(TSharedPtr< const FUICommandList > InCommandList, FMultiBoxCustomization InCustomization, TSharedPtr<FExtender> InExtender, EOrientation Orientation, const bool InForceSmallIcons = false, const bool bUniform = false)
 		: FMultiBoxBuilder(bUniform ? EMultiBoxType::UniformToolBar : (Orientation == Orient_Horizontal) ? EMultiBoxType::ToolBar : EMultiBoxType::VerticalToolBar, InCustomization, false, InCommandList, InExtender)
 		, bSectionNeedsToBeApplied(false)
+		, bSectionShouldHaveSeparator(true)
 		, bIsFocusable(true)
 		, bForceSmallIcons(InForceSmallIcons)
 	{
@@ -521,6 +522,7 @@ public:
 	FToolBarBuilder(TSharedPtr<const FUICommandList> InCommandList, FMultiBoxCustomization InCustomization, TSharedPtr<FExtender> InExtender = nullptr, const bool InForceSmallIcons = false)
 		: FMultiBoxBuilder(EMultiBoxType::ToolBar, InCustomization, false, InCommandList, InExtender)
 		, bSectionNeedsToBeApplied(false)
+		, bSectionShouldHaveSeparator(true)
 		, bIsFocusable(true)
 		, bForceSmallIcons(InForceSmallIcons)
 	{
@@ -612,13 +614,15 @@ public:
 	 * Adds a toolbar separator
 	 */
 	SLATE_API void AddSeparator(FName InExtensionHook = NAME_None);
-	
+
 	/**
 	 * Starts a section on to the extender section hook stack
-	 * 
+	 *
 	 * @param InExtensionHook	The section hook. Can be NAME_None
+	 * @param bInSectionShouldHaveSeparator	Whether this section should have a separator to its left (this doesn't apply
+	 * to the first section added, because that never gets a separator).
 	 */
-	SLATE_API void BeginSection( FName InExtensionHook );
+	SLATE_API void BeginSection(FName InExtensionHook, bool bInSectionShouldHaveSeparator = true);
 
 	/**
 	 * Ends the current section
@@ -646,6 +650,7 @@ protected:
 	FToolBarBuilder(EMultiBoxType InType, TSharedPtr<const FUICommandList> InCommandList, FMultiBoxCustomization InCustomization, TSharedPtr<FExtender> InExtender = TSharedPtr<FExtender>(), const bool InForceSmallIcons = false)
 		: FMultiBoxBuilder(InType, InCustomization, false, InCommandList, InExtender)
 		, bSectionNeedsToBeApplied(false)
+		, bSectionShouldHaveSeparator(true)
 		, bIsFocusable(false)
 		, bForceSmallIcons(InForceSmallIcons)
 	{
@@ -669,6 +674,8 @@ private:
 	/** True if there is a pending section that needs to be applied */
 	bool bSectionNeedsToBeApplied;
 
+	bool bSectionShouldHaveSeparator;
+
 	/** Whether the buttons created can receive keyboard focus */
 	bool bIsFocusable;
 
