Add FToolMenuEntry::InitDynamicEntry

This makes it possible for helper code to construct ready-made dynamic entries. Before, such code needed access to the section the dynamic entry would be added to.

diff --git a/Engine/Source/Developer/ToolMenus/Private/ToolMenuEntry.cpp b/Engine/Source/Developer/ToolMenus/Private/ToolMenuEntry.cpp
--- a/Engine/Source/Developer/ToolMenus/Private/ToolMenuEntry.cpp
+++ b/Engine/Source/Developer/ToolMenus/Private/ToolMenuEntry.cpp
@@ -151,6 +151,13 @@ FToolMenuEntry FToolMenuEntry::InitMenuEntry(const FName InName, const FToolUIAc
 	return Entry;
 }
 
+FToolMenuEntry FToolMenuEntry::InitDynamicEntry(const FName InName, const FNewToolMenuSectionDelegate& InConstruct)
+{
+	FToolMenuEntry Entry(UToolMenus::Get()->CurrentOwner(), InName, EMultiBlockType::MenuEntry);
+	Entry.Construct = InConstruct;
+	return Entry;
+}
+
 FToolMenuEntry FToolMenuEntry::InitSubMenu(const FName InName, const TAttribute<FText>& InLabel, const TAttribute<FText>& InToolTip, const FNewToolMenuChoice& InMakeMenu, const FToolUIActionChoice& InAction, const EUserInterfaceActionType InUserInterfaceActionType, bool bInOpenSubMenuOnClick, const TAttribute<FSlateIcon>& InIcon, const bool bInShouldCloseWindowAfterMenuSelection)
 {
 	FToolMenuEntry Entry(UToolMenus::Get()->CurrentOwner(), InName, EMultiBlockType::MenuEntry);
diff --git a/Engine/Source/Developer/ToolMenus/Private/ToolMenuSection.cpp b/Engine/Source/Developer/ToolMenus/Private/ToolMenuSection.cpp
--- a/Engine/Source/Developer/ToolMenus/Private/ToolMenuSection.cpp
+++ b/Engine/Source/Developer/ToolMenus/Private/ToolMenuSection.cpp
@@ -102,9 +102,7 @@ FToolMenuEntry& FToolMenuSection::AddMenuEntryWithCommandList(const TSharedPtr<
 
 FToolMenuEntry& FToolMenuSection::AddDynamicEntry(const FName InName, const FNewToolMenuSectionDelegate& InConstruct)
 {
-	FToolMenuEntry& Entry = AddEntry(FToolMenuEntry(UToolMenus::Get()->CurrentOwner(), InName, EMultiBlockType::MenuEntry));
-	Entry.Construct = InConstruct;
-	return Entry;
+	return AddEntry(FToolMenuEntry::InitDynamicEntry(InName, InConstruct));
 }
 
 FToolMenuEntry& FToolMenuSection::AddDynamicEntry(const FName InName, const FNewToolMenuDelegateLegacy& InConstruct)
diff --git a/Engine/Source/Developer/ToolMenus/Public/ToolMenuEntry.h b/Engine/Source/Developer/ToolMenus/Public/ToolMenuEntry.h
--- a/Engine/Source/Developer/ToolMenus/Public/ToolMenuEntry.h
+++ b/Engine/Source/Developer/ToolMenus/Public/ToolMenuEntry.h
@@ -123,6 +123,7 @@ struct TOOLMENUS_API FToolMenuEntry
 	static FToolMenuEntry InitMenuEntry(const FName InNameOverride, const TSharedPtr< const FUICommandInfo >& InCommand, const TAttribute<FText>& InLabelOverride = TAttribute<FText>(), const TAttribute<FText>& InToolTipOverride = TAttribute<FText>(), const TAttribute<FSlateIcon>& InIconOverride = TAttribute<FSlateIcon>(), const FName InTutorialHighlightName = NAME_None);
 	static FToolMenuEntry InitMenuEntryWithCommandList(const TSharedPtr<const FUICommandInfo>& InCommand, const TSharedPtr<const FUICommandList>& InCommandList, const TAttribute<FText>& InLabelOverride = TAttribute<FText>(), const TAttribute<FText>& InToolTipOverride = TAttribute<FText>(), const TAttribute<FSlateIcon>& InIconOverride = TAttribute<FSlateIcon>(), const FName InTutorialHighlightName = NAME_None, const TOptional<FName> InNameOverride = TOptional<FName>());
 	static FToolMenuEntry InitMenuEntry(const FName InName, const FToolUIActionChoice& InAction, const TSharedRef<SWidget>& Widget);
+	static FToolMenuEntry InitDynamicEntry(const FName InName, const FNewToolMenuSectionDelegate& InConstruct);
 
 	static FToolMenuEntry InitSubMenu(const FName InName, const TAttribute<FText>& InLabel, const TAttribute<FText>& InToolTip, const FNewToolMenuChoice& InMakeMenu, const FToolUIActionChoice& InAction, const EUserInterfaceActionType InUserInterfaceActionType, bool bInOpenSubMenuOnClick = false, const TAttribute<FSlateIcon>& InIcon = TAttribute<FSlateIcon>(), const bool bShouldCloseWindowAfterMenuSelection = true);
 	static FToolMenuEntry InitSubMenu(const FName InName, const TAttribute<FText>& InLabel, const TAttribute<FText>& InToolTip, const FNewToolMenuChoice& InMakeMenu, bool bInOpenSubMenuOnClick = false, const TAttribute<FSlateIcon>& InIcon = TAttribute<FSlateIcon>(), const bool bShouldCloseWindowAfterMenuSelection = true, const FName InTutorialHighlightName = NAME_None);
