Add new EToolMenuInsertType::Last

With this, code can specify that an entry should be positioned last in a section or a section positioned last in a menu. The biggest effect of this is that Last-positioned entries or sections will appear after Default-positioned entries or sections even if the Last-positioned ones are added earlier.

diff --git a/Engine/Source/Developer/ToolMenus/Private/ToolMenu.cpp b/Engine/Source/Developer/ToolMenus/Private/ToolMenu.cpp
--- a/Engine/Source/Developer/ToolMenus/Private/ToolMenu.cpp
+++ b/Engine/Source/Developer/ToolMenus/Private/ToolMenu.cpp
@@ -113,14 +113,28 @@ int32 UToolMenu::IndexOfSection(const FName InSectionName) const
 	return INDEX_NONE;
 }
 
+// Note: This function is very similar to FToolMenuSection::FindBlockInsertIndex.
 int32 UToolMenu::FindInsertIndex(const FToolMenuSection& InSection) const
 {
 	const FToolMenuInsert InInsertPosition = InSection.InsertPosition;
+
+	// Insert a Default-positioned section after all First and Default-positioned sections but before any
+	// Last-positioned sections.
 	if (InInsertPosition.IsDefault())
 	{
+		for (int32 i = 0; i < Sections.Num(); ++i)
+		{
+			if (Sections[i].InsertPosition.Position == EToolMenuInsertType::Last)
+			{
+				return i;
+			}
+		}
+
 		return Sections.Num();
 	}
 
+	// Insert a First-positioned section after any other First-positioned sections but before all Default and
+	// Last-positioned sections.
 	if (InInsertPosition.Position == EToolMenuInsertType::First)
 	{
 		for (int32 i = 0; i < Sections.Num(); ++i)
@@ -134,6 +148,20 @@ int32 UToolMenu::FindInsertIndex(const FToolMenuSection& InSection) const
 		return Sections.Num();
 	}
 
+	// Insert a Last-positioned section after all other sections, include other Last-positioned sections.
+	if (InInsertPosition.Position == EToolMenuInsertType::Last)
+	{
+		for (int32 i = Sections.Num() - 1; i >= 0; --i)
+		{
+			if (Sections[i].InsertPosition.Position == InInsertPosition.Position)
+			{
+				return i + 1;
+			}
+		}
+
+		return Sections.Num();
+	}
+
 	int32 DestIndex = IndexOfSection(InInsertPosition.Name);
 	if (DestIndex == INDEX_NONE)
 	{
diff --git a/Engine/Source/Developer/ToolMenus/Private/ToolMenuSection.cpp b/Engine/Source/Developer/ToolMenus/Private/ToolMenuSection.cpp
--- a/Engine/Source/Developer/ToolMenus/Private/ToolMenuSection.cpp
+++ b/Engine/Source/Developer/ToolMenus/Private/ToolMenuSection.cpp
@@ -196,15 +196,28 @@ int32 FToolMenuSection::RemoveEntriesByOwner(const FToolMenuOwner InOwner)
 	return 0;
 }
 
+// Note: This function is very similar to UToolMenu::FindInsertIndex.
 int32 FToolMenuSection::FindBlockInsertIndex(const FToolMenuEntry& InBlock) const
 {
 	const FToolMenuInsert InPosition = InBlock.InsertPosition;
 
+	// Insert a Default-positioned entry after all First and Default-positioned entries but before any Last-positioned
+	// entries.
 	if (InPosition.IsDefault())
 	{
+		for (int32 i = 0; i < Blocks.Num(); ++i)
+		{
+			if (Blocks[i].InsertPosition.Position == EToolMenuInsertType::Last)
+			{
+				return i;
+			}
+		}
+
 		return Blocks.Num();
 	}
 
+	// Insert a First-positioned entry after any other First-positioned entries but before all Default and
+	// Last-positioned entries.
 	if (InPosition.Position == EToolMenuInsertType::First)
 	{
 		for (int32 i = 0; i < Blocks.Num(); ++i)
@@ -218,6 +231,20 @@ int32 FToolMenuSection::FindBlockInsertIndex(const FToolMenuEntry& InBlock) cons
 		return Blocks.Num();
 	}
 
+	// Insert a Last-positioned entry after all other entries, include other Last-positioned entries.
+	if (InPosition.Position == EToolMenuInsertType::Last)
+	{
+		for (int32 i = Blocks.Num() - 1; i >= 0; --i)
+		{
+			if (Blocks[i].InsertPosition.Position == InPosition.Position)
+			{
+				return i + 1;
+			}
+		}
+
+		return Blocks.Num();
+	}
+
 	int32 DestIndex = IndexOfBlock(InPosition.Name);
 	if (DestIndex == INDEX_NONE)
 	{
diff --git a/Engine/Source/Developer/ToolMenus/Public/ToolMenuMisc.h b/Engine/Source/Developer/ToolMenus/Public/ToolMenuMisc.h
--- a/Engine/Source/Developer/ToolMenus/Public/ToolMenuMisc.h
+++ b/Engine/Source/Developer/ToolMenus/Public/ToolMenuMisc.h
@@ -60,7 +60,8 @@ enum class EToolMenuInsertType : uint8
 	Default,
 	Before,
 	After,
-	First
+	First,
+	Last
 };
 
 USTRUCT(BlueprintType)
