When raising children of a ToolMenu toolbar submenu, do not also recurse and raise their children

This prevents duplicated raises entries in toolbars when a child submenu is raised and a child of that child (grandchild) is also raised. With this change, the child and the grandchild both appear once in the top-level toolbar. Before, the grandchild would appear twice, raised once via the parent and once via the child.

diff --git a/Engine/Source/Developer/ToolMenus/Private/ToolMenus.cpp b/Engine/Source/Developer/ToolMenus/Private/ToolMenus.cpp
--- a/Engine/Source/Developer/ToolMenus/Private/ToolMenus.cpp
+++ b/Engine/Source/Developer/ToolMenus/Private/ToolMenus.cpp
@@ -1583,7 +1583,10 @@ void UToolMenus::PopulateToolBarBuilderWithEntry(
 				Block.ToolBarData.bSimpleComboBox, Block.TutorialHighlightName);
 
 			// Also add any top-level flagged children to the toolbar.
-			PopulateToolBarBuilderWithTopLevelChildren(ToolBarBuilder, MenuData, Block);
+			if (!bIsRaisingToTopLevel)
+			{
+				PopulateToolBarBuilderWithTopLevelChildren(ToolBarBuilder, MenuData, Block);
+			}
 		}
 	}
 	else if (Block.Type == EMultiBlockType::Separator)
