Add support for toolbar override labels to ToolMenu entries

This is useful when raising a menu entry to the top-level toolbar. The original label is used for the menu entry and the override for the toolbar.

diff --git a/Engine/Source/Developer/ToolMenus/Private/ToolMenus.cpp b/Engine/Source/Developer/ToolMenus/Private/ToolMenus.cpp
--- a/Engine/Source/Developer/ToolMenus/Private/ToolMenus.cpp
+++ b/Engine/Source/Developer/ToolMenus/Private/ToolMenus.cpp
@@ -1518,6 +1518,16 @@ void UToolMenus::PopulateToolBarBuilderWithEntry(
 
 	if (Block.Type == EMultiBlockType::ToolBarButton || (Block.Type == EMultiBlockType::MenuEntry && !Block.IsSubMenu()))
 	{
+		TAttribute<FText> Label;
+		if (Block.ToolbarLabelOverride.IsSet())
+		{
+			Label = Block.ToolbarLabelOverride;
+		}
+		else
+		{
+			Label = Block.Label;
+		}
+
 		if (Block.Command.IsValid() && !Block.IsCommandKeybindOnly())
 		{
 			bool bPopCommandList = false;
@@ -1535,7 +1545,8 @@ void UToolMenus::PopulateToolBarBuilderWithEntry(
 			}
 
 			ToolBarBuilder.AddToolBarButton(
-				Block.Command, Block.Name, Block.Label, Block.ToolTip, Block.Icon, Block.TutorialHighlightName);
+				Block.Command, Block.Name, Label, Block.ToolTip, Block.Icon, Block.TutorialHighlightName
+			);
 
 			if (bPopCommandList)
 			{
@@ -1553,34 +1564,55 @@ void UToolMenus::PopulateToolBarBuilderWithEntry(
 		}
 		else
 		{
-			ToolBarBuilder.AddToolBarButton(UIAction, Block.Name, Block.Label, Block.ToolTip, Block.Icon,
-				Block.UserInterfaceActionType, Block.TutorialHighlightName);
+			ToolBarBuilder.AddToolBarButton(
+				UIAction, Block.Name, Label, Block.ToolTip, Block.Icon, Block.UserInterfaceActionType, Block.TutorialHighlightName
+			);
 		}
 
 		if (Block.ToolBarData.OptionsDropdownData.IsValid())
 		{
 			FOnGetContent OnGetContent = ConvertWidgetChoice(
 				Block.ToolBarData.OptionsDropdownData->MenuContentGenerator, MenuData->Context);
-			ToolBarBuilder.AddComboButton(Block.ToolBarData.OptionsDropdownData->Action, OnGetContent, Block.Label,
-				Block.ToolBarData.OptionsDropdownData->ToolTip, Block.Icon, true, Block.TutorialHighlightName);
+			ToolBarBuilder.AddComboButton(
+				Block.ToolBarData.OptionsDropdownData->Action,
+				OnGetContent,
+				Label,
+				Block.ToolBarData.OptionsDropdownData->ToolTip,
+				Block.Icon,
+				true,
+				Block.TutorialHighlightName
+			);
 		}
 	}
 	else if (Block.Type == EMultiBlockType::ToolBarComboButton
 			 || (Block.Type == EMultiBlockType::MenuEntry && Block.IsSubMenu()))
 	{
+		TAttribute<FText> Label;
+		if (Block.ToolbarLabelOverride.IsSet())
+		{
+			Label = Block.ToolbarLabelOverride;
+		}
+		else
+		{
+			Label = Block.Label;
+		}
+
 		FOnGetContent OnGetContent = ConvertWidgetChoice(
 			Block.ToolBarData.ComboButtonContextMenuGenerator, MenuData->Context);
 		if (OnGetContent.IsBound())
 		{
-			ToolBarBuilder.AddComboButton(UIAction, OnGetContent, Block.Label, Block.ToolTip, Block.Icon,
-				Block.ToolBarData.bSimpleComboBox, Block.TutorialHighlightName);
+			ToolBarBuilder.AddComboButton(
+				UIAction, OnGetContent, Label, Block.ToolTip, Block.Icon, Block.ToolBarData.bSimpleComboBox, Block.TutorialHighlightName
+			);
 		}
 		else
 		{
 			FOnGetContent Delegate = FOnGetContent::CreateUObject(
 				this, &UToolMenus::GenerateToolbarComboButtonMenu, TWeakObjectPtr<UToolMenu>(MenuData), Block.Name);
-			ToolBarBuilder.AddComboButton(UIAction, Delegate, Block.Label, Block.ToolTip, Block.Icon,
-				Block.ToolBarData.bSimpleComboBox, Block.TutorialHighlightName);
+
+			ToolBarBuilder.AddComboButton(
+				UIAction, Delegate, Label, Block.ToolTip, Block.Icon, Block.ToolBarData.bSimpleComboBox, Block.TutorialHighlightName
+			);
 
 			// Also add any top-level flagged children to the toolbar.
 			if (!bIsRaisingToTopLevel)
diff --git a/Engine/Source/Developer/ToolMenus/Public/ToolMenuEntry.h b/Engine/Source/Developer/ToolMenus/Public/ToolMenuEntry.h
--- a/Engine/Source/Developer/ToolMenus/Public/ToolMenuEntry.h
+++ b/Engine/Source/Developer/ToolMenus/Public/ToolMenuEntry.h
@@ -225,6 +225,9 @@ public:
 	TAttribute<FSlateIcon> Icon;
 	TAttribute<FText> InputBindingLabel;
 
+	/** Optional override label to use when the entry appears in a toolbar. */
+	TAttribute<FText> ToolbarLabelOverride;
+
 private:
 
 	friend class UToolMenus;
