Entries raised to the top-level of ToolMenu toolbars no longer get a label if they have an icon

This is to reduce the visual load when raising items to the top-level that already have an icon to identify them.

diff --git a/Engine/Source/Developer/ToolMenus/Private/ToolMenus.cpp b/Engine/Source/Developer/ToolMenus/Private/ToolMenus.cpp
--- a/Engine/Source/Developer/ToolMenus/Private/ToolMenus.cpp
+++ b/Engine/Source/Developer/ToolMenus/Private/ToolMenus.cpp
@@ -1516,18 +1516,27 @@ void UToolMenus::PopulateToolBarBuilderWithEntry(
 		);
 	}
 
-	if (Block.Type == EMultiBlockType::ToolBarButton || (Block.Type == EMultiBlockType::MenuEntry && !Block.IsSubMenu()))
+	TAttribute<FText> Label;
 	{
-		TAttribute<FText> Label;
+		const bool bHasIcon = Block.Icon.IsSet() || (Block.Command.IsValid() && Block.Command->GetIcon().IsSet());
+
 		if (Block.ToolbarLabelOverride.IsSet())
 		{
 			Label = Block.ToolbarLabelOverride;
 		}
-		else
+		else if (!bIsRaisingToTopLevel || (bIsRaisingToTopLevel && !bHasIcon))
 		{
 			Label = Block.Label;
 		}
+		else
+		{
+			// This explicitly sets it empty, which will override the label on commands, for example.
+			Label = FText();
+		}
+	}
 
+	if (Block.Type == EMultiBlockType::ToolBarButton || (Block.Type == EMultiBlockType::MenuEntry && !Block.IsSubMenu()))
+	{
 		if (Block.Command.IsValid() && !Block.IsCommandKeybindOnly())
 		{
 			bool bPopCommandList = false;
@@ -1587,16 +1596,6 @@ void UToolMenus::PopulateToolBarBuilderWithEntry(
 	else if (Block.Type == EMultiBlockType::ToolBarComboButton
 			 || (Block.Type == EMultiBlockType::MenuEntry && Block.IsSubMenu()))
 	{
-		TAttribute<FText> Label;
-		if (Block.ToolbarLabelOverride.IsSet())
-		{
-			Label = Block.ToolbarLabelOverride;
-		}
-		else
-		{
-			Label = Block.Label;
-		}
-
 		FOnGetContent OnGetContent = ConvertWidgetChoice(
 			Block.ToolBarData.ComboButtonContextMenuGenerator, MenuData->Context);
 		if (OnGetContent.IsBound())
