Raise separators if a previous and a following entry is also raised

diff --git a/Engine/Source/Developer/ToolMenus/Private/ToolMenus.cpp b/Engine/Source/Developer/ToolMenus/Private/ToolMenus.cpp
--- a/Engine/Source/Developer/ToolMenus/Private/ToolMenus.cpp
+++ b/Engine/Source/Developer/ToolMenus/Private/ToolMenus.cpp
@@ -1413,7 +1413,8 @@ void UToolMenus::PopulateMenuBuilder(FMenuBuilder& MenuBuilder, UToolMenu* MenuD
 }
 
 void UToolMenus::ExtractChildBlocksFromSubMenu(
-	UToolMenu* ParentMenu, FToolMenuEntry& InBlock, TArray<TPair<UToolMenu*, FToolMenuEntry*>>& SubMenuBlockPairs)
+	UToolMenu* ParentMenu, FToolMenuEntry& InBlock, TArray<TTuple<UToolMenu*, FToolMenuEntry*, FName>>& SubMenuBlocks
+)
 {
 	const bool bBuiltViaDelegate =
 		ConvertWidgetChoice(InBlock.ToolBarData.ComboButtonContextMenuGenerator, ParentMenu->Context).IsBound();
@@ -1441,7 +1442,7 @@ void UToolMenus::ExtractChildBlocksFromSubMenu(
 		for (int j = Section.Blocks.Num() - 1; j >= 0; --j)
 		{
 			FToolMenuEntry& Block = Section.Blocks[j];
-			SubMenuBlockPairs.Push(TPair<UToolMenu*, FToolMenuEntry*>(SubMenu, &Block));
+			SubMenuBlocks.Push(TTuple<UToolMenu*, FToolMenuEntry*, FName>(SubMenu, &Block, Section.Name));
 		}
 	}
 }
@@ -1449,16 +1450,24 @@ void UToolMenus::ExtractChildBlocksFromSubMenu(
 void UToolMenus::PopulateToolBarBuilderWithTopLevelChildren(
 	FToolBarBuilder& ToolBarBuilder, UToolMenu* ParentMenu, FToolMenuEntry& InBlock)
 {
-	TArray<TPair<UToolMenu*, FToolMenuEntry*>> SubMenuBlockPairs;
+	TArray<TTuple<UToolMenu*, FToolMenuEntry*, FName>> SubMenuBlocks;
 	// Seed the blocks with the passed-in submenu.
-	ExtractChildBlocksFromSubMenu(ParentMenu, InBlock, SubMenuBlockPairs);
+	ExtractChildBlocksFromSubMenu(ParentMenu, InBlock, SubMenuBlocks);
+
+	// Collect the blocks that might be raised to the top level. Also include separators so we can visualize those in
+	// the toolbar when they appear between two raised blocks.
+	TArray<TTuple<UToolMenu*, FToolMenuEntry*, FName>> BlocksToAdd;
 
+	// Traverse the sub menu blocks we've found thus far to find more grandchild blocks that are raised (raised set to
+	// boolean true) or could be dynamically raised (raised set to a TAttribute<bool> driven by a delegate) to the
+	// top-level toolbar.
 	int32 NumIterations = 0;
-	while (SubMenuBlockPairs.Num() > 0)
+	while (SubMenuBlocks.Num() > 0)
 	{
-		const TPair<UToolMenu*, FToolMenuEntry*> SubMenuAndBlock = SubMenuBlockPairs.Pop(EAllowShrinking::No);
-		UToolMenu* const SubMenu = SubMenuAndBlock.Key;
-		FToolMenuEntry* const Block = SubMenuAndBlock.Value;
+		const TTuple<UToolMenu*, FToolMenuEntry*, FName> SubMenuBlock = SubMenuBlocks.Pop(EAllowShrinking::No);
+		UToolMenu* const SubMenu = SubMenuBlock.Get<0>();
+		FToolMenuEntry* const Block = SubMenuBlock.Get<1>();
+		const FName SectionName = SubMenuBlock.Get<2>();
 
 		// Keep track of how many blocks we've visited to ensure we don't loop indefinitely.
 		if (++NumIterations > 5000)
@@ -1472,13 +1481,73 @@ void UToolMenus::PopulateToolBarBuilderWithTopLevelChildren(
 		const bool bIsBound = Block->ShowInToolbarTopLevel.IsBound();
 		const bool bIsSetToValue = !bIsBound && Block->ShowInToolbarTopLevel.IsSet();
 		const bool bIsSetToTrueValue = bIsSetToValue && Block->ShowInToolbarTopLevel.Get();
-		if (bIsBound || bIsSetToTrueValue)
+		const bool bIsBoundOrTrue = bIsBound || bIsSetToTrueValue;
+		if (bIsBoundOrTrue || Block->Type == EMultiBlockType::Separator)
+		{
+			BlocksToAdd.Add(SubMenuBlock);
+		}
+
+		ExtractChildBlocksFromSubMenu(SubMenu, *Block, SubMenuBlocks);
+	}
+
+	// Do not allow leading separators.
+	while (BlocksToAdd.Num() > 0 && BlocksToAdd[0].Get<1>()->Type == EMultiBlockType::Separator)
+	{
+		BlocksToAdd.RemoveAt(0);
+	}
+
+	// Do not allow trailing separators.
+	while (BlocksToAdd.Num() > 0 && BlocksToAdd[BlocksToAdd.Num() - 1].Get<1>()->Type == EMultiBlockType::Separator)
+	{
+		BlocksToAdd.RemoveAt(BlocksToAdd.Num() - 1);
+	}
+
+	// Do not allow rows of separators.
+	for (int i = 1; i < BlocksToAdd.Num(); ++i)
+	{
+		const bool bIsCurrentSeparator = BlocksToAdd[i].Get<1>()->Type == EMultiBlockType::Separator;
+		const bool bWasPreviousSeparator = BlocksToAdd[i - 1].Get<1>()->Type == EMultiBlockType::Separator;
+		const bool bPartOfRowOfSeparators = bIsCurrentSeparator && bWasPreviousSeparator;
+
+		if (bPartOfRowOfSeparators)
+		{
+			BlocksToAdd.RemoveAt(i--);
+		}
+	}
+
+	if (BlocksToAdd.IsEmpty())
+	{
+		return;
+	}
+
+	bool bHasRaisedEntrySinceLastSeparator = false;
+	// Seed the previous section with the first blocks's section so we don't start with adding a separator because
+	// sections seem to have changed.
+	FName PreviousSectionName = BlocksToAdd[0].Get<2>();
+	for (int i = 0; i < BlocksToAdd.Num(); ++i)
+	{
+		UToolMenu* const SubMenu = BlocksToAdd[i].Get<0>();
+		FToolMenuEntry* const Entry = BlocksToAdd[i].Get<1>();
+		const FName SectionName = BlocksToAdd[i].Get<2>();
+
+		// Add a separator if one was found or a new section was encountered.
+		if (SectionName != PreviousSectionName
+			|| (Entry->Type == EMultiBlockType::Separator && bHasRaisedEntrySinceLastSeparator))
+		{
+			const FName UnsetExtensionHook = NAME_None;
+			ToolBarBuilder.AddSeparator(UnsetExtensionHook);
+			bHasRaisedEntrySinceLastSeparator = false;
+		}
+
+		// Make sure we actually add the entry if the reason we added a separator above was that the section names changed.
+		if (Entry->Type != EMultiBlockType::Separator)
 		{
 			const bool bRaiseToTopLevel = true;
-			PopulateToolBarBuilderWithEntry(ToolBarBuilder, SubMenu, *Block, bRaiseToTopLevel);
+			PopulateToolBarBuilderWithEntry(ToolBarBuilder, SubMenu, *Entry, bRaiseToTopLevel);
+			bHasRaisedEntrySinceLastSeparator = true;
 		}
 
-		ExtractChildBlocksFromSubMenu(SubMenu, *Block, SubMenuBlockPairs);
+		PreviousSectionName = SectionName;
 	}
 }
 
diff --git a/Engine/Source/Developer/ToolMenus/Public/ToolMenus.h b/Engine/Source/Developer/ToolMenus/Public/ToolMenus.h
--- a/Engine/Source/Developer/ToolMenus/Public/ToolMenus.h
+++ b/Engine/Source/Developer/ToolMenus/Public/ToolMenus.h
@@ -387,7 +387,8 @@ private:
 	void PopulateMenuBarBuilder(FMenuBarBuilder& MenuBarBuilder, UToolMenu* MenuData);
 
 	void ExtractChildBlocksFromSubMenu(
-		UToolMenu* ParentMenu, FToolMenuEntry& InBlock, TArray<TPair<UToolMenu*, FToolMenuEntry*>>& SubMenuBlockPairs);
+		UToolMenu* ParentMenu, FToolMenuEntry& InBlock, TArray<TTuple<UToolMenu*, FToolMenuEntry*, FName>>& SubMenuBlocks
+	);
 	void PopulateToolBarBuilderWithTopLevelChildren(
 		FToolBarBuilder& ToolBarBuilder, UToolMenu* ParentMenu, FToolMenuEntry& InBlock);
 	void PopulateToolBarBuilderWithEntry(
diff --git a/Engine/Source/Editor/EditorStyle/Private/StarshipStyle.cpp b/Engine/Source/Editor/EditorStyle/Private/StarshipStyle.cpp
--- a/Engine/Source/Editor/EditorStyle/Private/StarshipStyle.cpp
+++ b/Engine/Source/Editor/EditorStyle/Private/StarshipStyle.cpp
@@ -2453,10 +2453,16 @@ void FStarshipEditorStyle::FStyle::SetupViewportStyles()
 					.SetPadding(FMargin(4));
 			}
 
+			FSlateBrush SeparatorBrush;
+			SeparatorBrush.DrawAs = ESlateBrushDrawType::Box;
+			SeparatorBrush.TintColor = FStyleColors::Dropdown;
+
 			ToolBarStyle.SetIconSize(Icon16x16)
 				.SetButtonStyle(ButtonStyle)
 				.SetToggleButtonStyle(ToggleButtonStyle)
-				.SetComboButtonPadding(FMargin(4.f, 0.0f));
+				.SetComboButtonPadding(FMargin(4.f, 0.0f))
+				.SetSeparatorBrush(SeparatorBrush)
+				.SetSeparatorPadding(FMargin(0, 0.0f));
 
 			Set("ViewportToolbar", ToolBarStyle);
 		}
