Implement proper visibility of new viewport toolbar separators

Set the visibility of separators to a delegate that depends on the visibility of previously and upcoming raised entries. This allows separators to only be visible when they are surrounded by raised entries and avoids situations where dynamically raised entries are hidden, leaving separators as the first or last elements in a toolbar.

This is needed for the new viewport toolbar. We raise entries from toolbar submenus (combo buttons) so they appear directly in the top-level viewport toolbar. The condition for raising entries can be dynamically driven by a TAttribute<bool> delegate. We want to add separators between raised entries if said entries either live in separate sections or have an explicit separator between them. Because entries are dynamically visible but we still want to show separators when the entries are visible, the visibility of the separators themselves has to be dynamic which this change supports.

diff --git a/Engine/Source/Developer/ToolMenus/Private/ToolMenus.cpp b/Engine/Source/Developer/ToolMenus/Private/ToolMenus.cpp
--- a/Engine/Source/Developer/ToolMenus/Private/ToolMenus.cpp
+++ b/Engine/Source/Developer/ToolMenus/Private/ToolMenus.cpp
@@ -1521,6 +1521,31 @@ void UToolMenus::PopulateToolBarBuilderWithTopLevelChildren(
 		return;
 	}
 
+	// Dynamic visibility of separators
+	//
+	// We add separators in the top-level toolbar between raised entries if the raised entries lived in different
+	// sections or if a separator was explicitly added between the raised entries.
+	//
+	// Since entries can be dynamically raised, added toolbar separators must have dynamic visiblity. To support this,
+	// we record the visibility delegates of previously raised entries so separator visiblity delegates can use them.
+	//
+	// A menu might look like this:
+	//
+	//  |-- previous1 -|               |-- previous2 -|               |-- previous3 -|             |---- next ----|
+	//  raisedA raisedB SEPARATOR(N-2) raisedC raisedD SEPARATOR(N-1) raisedE raisedF SEPARATOR(N) raisedG raisedH
+	//
+	// PreviousEntries = (raisedA, raisedB, raisedC, raisedD, raisedE, raisedF)
+	// NextEntries = (raisedG, raisedH)
+	//
+	// Separator visibility is then calculated like this:
+	//
+	//   sep_vis = anyVisible(nextEntries) && anyVisible(PreviousEntries
+	//
+
+	TArray<TAttribute<bool>> PreviousVisibilities;
+	// This has to be heap allocated so we can still add do it after a separator visibility delegate captures it.
+	TSharedPtr<TArray<TAttribute<bool>>> NextVisibilities = MakeShared<TArray<TAttribute<bool>>>();
+
 	bool bHasRaisedEntrySinceLastSeparator = false;
 	// Seed the previous section with the first blocks's section so we don't start with adding a separator because
 	// sections seem to have changed.
@@ -1532,11 +1557,48 @@ void UToolMenus::PopulateToolBarBuilderWithTopLevelChildren(
 		const FName SectionName = BlocksToAdd[i].Get<2>();
 
 		// Add a separator if one was found or a new section was encountered.
-		if (SectionName != PreviousSectionName
-			|| (Entry->Type == EMultiBlockType::Separator && bHasRaisedEntrySinceLastSeparator))
+		if (bHasRaisedEntrySinceLastSeparator
+			&& (SectionName != PreviousSectionName || Entry->Type == EMultiBlockType::Separator))
 		{
+			// Step entry visibility delegate records forward now that we encountered a new separator.
+			PreviousVisibilities.Append(*NextVisibilities);
+			NextVisibilities = MakeShared<TArray<TAttribute<bool>>>();
+
+			const TAttribute<EVisibility> VisibilityOverride = TAttribute<EVisibility>::CreateLambda(
+				[Previous = PreviousVisibilities, Next = NextVisibilities]()
+				{
+					// This function calculates this expression and earlies out if possible.
+					// const bool bVisible = bAnyNext && bAnyPrevious;
+
+					bool bAnyNext = false;
+					for (const TAttribute<bool>& Visibility : *Next)
+					{
+						if (Visibility.Get())
+						{
+							bAnyNext = true;
+							break;
+						}
+					}
+
+					if (!bAnyNext)
+					{
+						return EVisibility::Collapsed;
+					}
+
+					for (const TAttribute<bool>& Visibility : Previous)
+					{
+						if (Visibility.Get())
+						{
+							return EVisibility::Visible;
+						}
+					}
+
+					return EVisibility::Collapsed;
+				}
+			);
+
 			const FName UnsetExtensionHook = NAME_None;
-			ToolBarBuilder.AddSeparator(UnsetExtensionHook);
+			ToolBarBuilder.AddSeparator(UnsetExtensionHook, VisibilityOverride);
 			bHasRaisedEntrySinceLastSeparator = false;
 		}
 
@@ -1546,6 +1608,9 @@ void UToolMenus::PopulateToolBarBuilderWithTopLevelChildren(
 			const bool bRaiseToTopLevel = true;
 			PopulateToolBarBuilderWithEntry(ToolBarBuilder, SubMenu, *Entry, bRaiseToTopLevel);
 			bHasRaisedEntrySinceLastSeparator = true;
+
+			// Keep track of added entries' visibilities so separators can set their visibility override.
+			NextVisibilities->Add(Entry->ShowInToolbarTopLevel);
 		}
 
 		PreviousSectionName = SectionName;
