Make FToolMenuSection run a check against IsValid() on commands and command lists when adding menu entries

diff --git a/Engine/Source/Developer/ToolMenus/Private/ToolMenuSection.cpp b/Engine/Source/Developer/ToolMenus/Private/ToolMenuSection.cpp
--- a/Engine/Source/Developer/ToolMenus/Private/ToolMenuSection.cpp
+++ b/Engine/Source/Developer/ToolMenus/Private/ToolMenuSection.cpp
@@ -87,16 +87,23 @@ FToolMenuEntry& FToolMenuSection::AddMenuEntry(const FName InName, const TAttrib
 
 FToolMenuEntry& FToolMenuSection::AddMenuEntry(const TSharedPtr< const FUICommandInfo >& InCommand, const TAttribute<FText>& InLabelOverride, const TAttribute<FText>& InToolTipOverride, const TAttribute<FSlateIcon>& InIconOverride, const FName InTutorialHighlightName, const TOptional<FName> InNameOverride)
 {
+	check(InCommand.IsValid());
+
 	return AddEntry(FToolMenuEntry::InitMenuEntry(InCommand, InLabelOverride, InToolTipOverride, InIconOverride, InTutorialHighlightName, InNameOverride));
 }
 
 FToolMenuEntry& FToolMenuSection::AddMenuEntry(const FName InNameOverride, const TSharedPtr< const FUICommandInfo >& InCommand, const TAttribute<FText>& InLabelOverride, const TAttribute<FText>& InToolTipOverride, const TAttribute<FSlateIcon>& InIconOverride, const FName InTutorialHighlightName)
 {
+	check(InCommand.IsValid());
+
 	return AddEntry(FToolMenuEntry::InitMenuEntry(InCommand, InLabelOverride, InToolTipOverride, InIconOverride, InTutorialHighlightName, InNameOverride));
 }
 
 FToolMenuEntry& FToolMenuSection::AddMenuEntryWithCommandList(const TSharedPtr< const FUICommandInfo >& InCommand, const TSharedPtr< const FUICommandList >& InCommandList, const TAttribute<FText>& InLabelOverride, const TAttribute<FText>& InToolTipOverride, const TAttribute<FSlateIcon>& InIconOverride, const FName InTutorialHighlightName, const TOptional<FName> InNameOverride)
 {
+	check(InCommand.IsValid());
+	check(InCommandList.IsValid());
+
 	return AddEntry(FToolMenuEntry::InitMenuEntryWithCommandList(InCommand, InCommandList, InLabelOverride, InToolTipOverride, InIconOverride, InTutorialHighlightName, InNameOverride));
 }
 
