Fix bug in ToolMenus where superfluous SSpacer widgets were added to toolbars after the last section

diff --git a/Engine/Source/Developer/ToolMenus/Private/ToolMenus.cpp b/Engine/Source/Developer/ToolMenus/Private/ToolMenus.cpp
--- a/Engine/Source/Developer/ToolMenus/Private/ToolMenus.cpp
+++ b/Engine/Source/Developer/ToolMenus/Private/ToolMenus.cpp
@@ -1891,15 +1891,21 @@ void UToolMenus::PopulateToolBarBuilder(FToolBarBuilder& ToolBarBuilder, UToolMe
 															  EToolMenuSectionAlign::Default,
 															  EToolMenuSectionAlign::Middle,
 															  EToolMenuSectionAlign::Last };
-	bool bDidAddSection = false;
+	int32 NumAddedSections = 0;
 	for (const EToolMenuSectionAlign CurrentAlignment : SectionAlignments)
 	{
+		// If we've already added all sections, there's no more work to do. Especially, we don't want to add any more SSpacers.
+		if (NumAddedSections >= MenuData->Sections.Num())
+		{
+			break;
+		}
+
 		const bool bIsMiddleOrLast = CurrentAlignment == EToolMenuSectionAlign::Middle
 								  || CurrentAlignment == EToolMenuSectionAlign::Last;
 
 		// Add a spacer before the middle and last alignment groups, and only if we've already added a section to a
 		// previous alignment group.
-		if (bIsMiddleOrLast && bDidAddSection)
+		if (bIsMiddleOrLast && NumAddedSections > 0)
 		{
 			FMenuEntryStyleParams StyleParams;
 			StyleParams.HorizontalAlignment = HAlign_Right;
@@ -1912,7 +1918,8 @@ void UToolMenus::PopulateToolBarBuilder(FToolBarBuilder& ToolBarBuilder, UToolMe
 		// Keep track if this is the first section we're adding for the CurrentAlignment. Make an exception if the
 		// current alignment is Default and we already added a first-aligned section, because in that case this
 		// isn't the first section in the "group" since we're displaying first and default-aligned sections together.
-		bool bFirstSectionInAlignmentGroup = CurrentAlignment == EToolMenuSectionAlign::Default ? !bDidAddSection : true;
+		bool bFirstSectionInAlignmentGroup = CurrentAlignment == EToolMenuSectionAlign::Default ? NumAddedSections == 0
+																								: true;
 		for (FToolMenuSection& Section : MenuData->Sections)
 		{
 			if (Section.Alignment != CurrentAlignment)
@@ -1936,7 +1943,7 @@ void UToolMenus::PopulateToolBarBuilder(FToolBarBuilder& ToolBarBuilder, UToolMe
 
 			ToolBarBuilder.EndSection();
 
-			bDidAddSection = true;
+			++NumAddedSections;
 			bFirstSectionInAlignmentGroup = false;
 		}
 	}
