[Viewport Toolbar] Fix viewport toolbar submenus disappearing after user performs unrealed editor actions

This fix issue where subsubmenus had their parents (also submenus) garbage collected before the subsubmenu was opened, resulting in empty menus because the subsubmenu couldn't be combined with its parent to build a menu (because the parent was garbage collected already).

The cause of this was how UToolMenus::ExtractChildBlocksFromSubMenu calls UToolMenus::GenerateSubMenu in an unconventional way, resulting in subsub+ menus being weakly referenced by combo buttons in toolbars, which could then not build their menus if their parents had been garbage collected.

diff --git a/Engine/Source/Developer/ToolMenus/Private/ToolMenus.cpp b/Engine/Source/Developer/ToolMenus/Private/ToolMenus.cpp
--- a/Engine/Source/Developer/ToolMenus/Private/ToolMenus.cpp
+++ b/Engine/Source/Developer/ToolMenus/Private/ToolMenus.cpp
@@ -1799,8 +1799,15 @@ void UToolMenus::PopulateToolBarBuilderWithEntry(
 		}
 		else
 		{
-			FOnGetContent Delegate = FOnGetContent::CreateUObject(
-				this, &UToolMenus::GenerateToolbarComboButtonMenu, TWeakObjectPtr<UToolMenu>(MenuData), Block.Name);
+			// Make sure to keep a strong reference to this submenu so it stays around until it is opened. This is
+			// needed here because we're could be a submenu of a submenu, so not even our parent is added to the Menus
+			// map and therefore our parent could have been gargage collected before this submenu is opened.
+			FOnGetContent Delegate = FOnGetContent::CreateLambda(
+				[this, StrongMenuData = TStrongObjectPtr<UToolMenu>(MenuData), BlockName = Block.Name]() -> TSharedRef<SWidget>
+				{
+					return this->GenerateToolbarComboButtonMenu(StrongMenuData.Get(), BlockName);
+				}
+			);
 
 			ToolBarBuilder.AddComboButton(
 				UIAction,
