Add toolbar-styling support for adding a padding to the right of raised children

* Add the float FToolBarStyle::RaisedChildrenRightPadding for specifying the padding.
* Update UToolMenus to respect this padding and add an SSpacer set to the padding with after the raised children. This SSpacer is not added for the very last group of raised children of the toolbar, to avoid dead space at the far right of the toolbar. The SSpacer has an overridden visibility delegate to make sure it is hidden if all raised children are, in case the children are raised dynamically.

diff --git a/Engine/Source/Developer/ToolMenus/Private/ToolMenus.cpp b/Engine/Source/Developer/ToolMenus/Private/ToolMenus.cpp
--- a/Engine/Source/Developer/ToolMenus/Private/ToolMenus.cpp
+++ b/Engine/Source/Developer/ToolMenus/Private/ToolMenus.cpp
@@ -13,6 +13,7 @@
 #include "UObject/WeakObjectPtrTemplates.h"
 
 #include "HAL/PlatformApplicationMisc.h" // For clipboard
+#include "Styling/ToolBarStyle.h"
 #include "Widgets/Layout/SScrollBox.h"
 #include "Widgets/Layout/SSpacer.h"
 
@@ -1448,7 +1449,8 @@ void UToolMenus::ExtractChildBlocksFromSubMenu(
 }
 
 void UToolMenus::PopulateToolBarBuilderWithTopLevelChildren(
-	FToolBarBuilder& ToolBarBuilder, UToolMenu* ParentMenu, FToolMenuEntry& InBlock)
+	FToolBarBuilder& ToolBarBuilder, UToolMenu* ParentMenu, FToolMenuEntry& InBlock, bool bAddSpaceAfterRaisedChildren
+)
 {
 	TArray<TTuple<UToolMenu*, FToolMenuEntry*, FName>> SubMenuBlocks;
 	// Seed the blocks with the passed-in submenu.
@@ -1520,6 +1522,9 @@ void UToolMenus::PopulateToolBarBuilderWithTopLevelChildren(
 		return;
 	}
 
+	// Dynamic visibility of trailing raised-children spacer.
+	TArray<TAttribute<bool>> AllRaisedVisibilities = TArray<TAttribute<bool>>();
+
 	// Dynamic visibility of separators
 	//
 	// We add separators in the top-level toolbar between raised entries if the raised entries lived in different
@@ -1610,14 +1615,52 @@ void UToolMenus::PopulateToolBarBuilderWithTopLevelChildren(
 
 			// Keep track of added entries' visibilities so separators can set their visibility override.
 			NextVisibilities->Add(Entry->ShowInToolbarTopLevel);
+
+			if (bAddSpaceAfterRaisedChildren)
+			{
+				AllRaisedVisibilities.Add(Entry->ShowInToolbarTopLevel);
+			}
 		}
 
 		PreviousSectionName = SectionName;
 	}
+
+	if (bAddSpaceAfterRaisedChildren)
+	{
+		const FMenuEntryStyleParams StyleParams;
+		const FName TutorialHighlightName = NAME_None;
+		const bool bSearchable = false;
+		const FNewMenuDelegate CustomMenuDelegate;
+
+		TAttribute<EVisibility> SpacerVisibilityOverride = TAttribute<EVisibility>::CreateLambda(
+			[AllRaisedVisibilities]()
+			{
+				for (const TAttribute<bool>& Visibility : AllRaisedVisibilities)
+				{
+					if (Visibility.Get())
+					{
+						return EVisibility::Visible;
+					}
+				}
+
+				return EVisibility::Collapsed;
+			}
+		);
+
+		const FToolBarStyle ToolbarStyle = ParentMenu->StyleSet->GetWidgetStyle<FToolBarStyle>(ParentMenu->StyleName);
+		ToolBarBuilder.AddWidget(
+			SNew(SSpacer).Size(FVector2D(ToolbarStyle.RaisedChildrenRightPadding, 0)),
+			StyleParams,
+			TutorialHighlightName,
+			bSearchable,
+			CustomMenuDelegate,
+			SpacerVisibilityOverride
+		);
+	}
 }
 
 void UToolMenus::PopulateToolBarBuilderWithEntry(
-	FToolBarBuilder& ToolBarBuilder, UToolMenu* MenuData, FToolMenuEntry& Block, bool bIsRaisingToTopLevel
+	FToolBarBuilder& ToolBarBuilder, UToolMenu* MenuData, FToolMenuEntry& Block, bool bIsRaisingToTopLevel, bool bIsLastBlockOfLastSection
 )
 {
 	if (Block.ToolBarData.ConstructLegacy.IsBound())
@@ -1826,7 +1869,7 @@ void UToolMenus::PopulateToolBarBuilderWithEntry(
 			// Also add any top-level flagged children to the toolbar.
 			if (!bIsRaisingToTopLevel)
 			{
-				PopulateToolBarBuilderWithTopLevelChildren(ToolBarBuilder, MenuData, Block);
+				PopulateToolBarBuilderWithTopLevelChildren(ToolBarBuilder, MenuData, Block, !bIsLastBlockOfLastSection);
 			}
 		}
 	}
@@ -1932,8 +1975,11 @@ void UToolMenus::PopulateToolBarBuilder(FToolBarBuilder& ToolBarBuilder, UToolMe
 		// isn't the first section in the "group" since we're displaying first and default-aligned sections together.
 		bool bFirstSectionInAlignmentGroup = CurrentAlignment == EToolMenuSectionAlign::Default ? NumAddedSections == 0
 																								: true;
-		for (FToolMenuSection& Section : MenuData->Sections)
+		for (int SectionIndex = 0; SectionIndex < MenuData->Sections.Num(); ++SectionIndex)
 		{
+			const bool bIsLastSection = SectionIndex >= MenuData->Sections.Num() - 1;
+			FToolMenuSection& Section = MenuData->Sections[SectionIndex];
+
 			if (Section.Alignment != CurrentAlignment)
 			{
 				continue;
@@ -1948,9 +1994,14 @@ void UToolMenus::PopulateToolBarBuilder(FToolBarBuilder& ToolBarBuilder, UToolMe
 			const bool bSectionShouldHaveSeparator = !bFirstSectionInAlignmentGroup;
 			ToolBarBuilder.BeginSection(Section.Name, bSectionShouldHaveSeparator);
 
-			for (FToolMenuEntry& Block : Section.Blocks)
+			for (int BlockIndex = 0; BlockIndex < Section.Blocks.Num(); ++BlockIndex)
 			{
-				PopulateToolBarBuilderWithEntry(ToolBarBuilder, MenuData, Block);
+				const bool bIsLastBlock = BlockIndex >= Section.Blocks.Num() - 1;
+				FToolMenuEntry& Block = Section.Blocks[BlockIndex];
+
+				const bool bRaiseToTopLevel = false;
+				const bool bIsLastBlockOfLastSection = bIsLastBlock && bIsLastSection;
+				PopulateToolBarBuilderWithEntry(ToolBarBuilder, MenuData, Block, bRaiseToTopLevel, bIsLastBlockOfLastSection);
 			}
 
 			ToolBarBuilder.EndSection();
diff --git a/Engine/Source/Developer/ToolMenus/Public/ToolMenus.h b/Engine/Source/Developer/ToolMenus/Public/ToolMenus.h
--- a/Engine/Source/Developer/ToolMenus/Public/ToolMenus.h
+++ b/Engine/Source/Developer/ToolMenus/Public/ToolMenus.h
@@ -390,9 +390,14 @@ private:
 		UToolMenu* ParentMenu, FToolMenuEntry& InBlock, TArray<TTuple<UToolMenu*, FToolMenuEntry*, FName>>& SubMenuBlocks
 	);
 	void PopulateToolBarBuilderWithTopLevelChildren(
-		FToolBarBuilder& ToolBarBuilder, UToolMenu* ParentMenu, FToolMenuEntry& InBlock);
+		FToolBarBuilder& ToolBarBuilder, UToolMenu* ParentMenu, FToolMenuEntry& InBlock, bool bAddSpaceAfterRaisedChildren
+	);
 	void PopulateToolBarBuilderWithEntry(
-		FToolBarBuilder& ToolBarBuilder, UToolMenu* MenuData, FToolMenuEntry& Entry, bool bIsRaisingToTopLevel = false
+		FToolBarBuilder& ToolBarBuilder,
+		UToolMenu* MenuData,
+		FToolMenuEntry& Entry,
+		bool bIsRaisingToTopLevel = false,
+		bool bIsLastBlockOfLastSection = false
 	);
 	void PopulateToolBarBuilder(FToolBarBuilder& ToolBarBuilder, UToolMenu* MenuData);
 
diff --git a/Engine/Source/Runtime/SlateCore/Public/Styling/ToolBarStyle.h b/Engine/Source/Runtime/SlateCore/Public/Styling/ToolBarStyle.h
--- a/Engine/Source/Runtime/SlateCore/Public/Styling/ToolBarStyle.h
+++ b/Engine/Source/Runtime/SlateCore/Public/Styling/ToolBarStyle.h
@@ -152,4 +152,8 @@ struct FToolBarStyle : public FSlateWidgetStyle
 	UPROPERTY(EditAnywhere, BlueprintReadWrite, Category = Appearance)
 	TOptional<TEnumAsByte<EVerticalAlignment>> VerticalAlignmentOverride;
 	FToolBarStyle& SetVerticalAlignment(const EVerticalAlignment& InVerticalAlignment) { VerticalAlignmentOverride = InVerticalAlignment; return *this; }
+
+	UPROPERTY(EditAnywhere, BlueprintReadWrite, Category = Appearance)
+	float RaisedChildrenRightPadding = 0.0f;
+	FToolBarStyle& SetRaisedChildrenRightPadding(float InRaisedChildrenRightPadding) { RaisedChildrenRightPadding = InRaisedChildrenRightPadding; return *this; }
 };
