Add ToolMenuContext DebugLogContextObjectClasses() scripting function so python devs can see what context objects are available

diff --git a/Engine/Source/Developer/ToolMenus/Private/ToolMenusBlueprintLibrary.cpp b/Engine/Source/Developer/ToolMenus/Private/ToolMenusBlueprintLibrary.cpp
--- a/Engine/Source/Developer/ToolMenus/Private/ToolMenusBlueprintLibrary.cpp
+++ b/Engine/Source/Developer/ToolMenus/Private/ToolMenusBlueprintLibrary.cpp
@@ -51,6 +51,14 @@ UObject* UToolMenuContextExtensions::FindByClass(const FToolMenuContext& Context
 	return Context.FindByClass(*InClass);
 }
 
+void UToolMenuContextExtensions::DebugLogContextObjectClasses(const FToolMenuContext& Context)
+{
+	const FString ContextObjectClasses = FString::JoinBy(Context.ContextObjects, TEXT(", "), [](const UObject* ContextObject) { 
+		return ContextObject ? ContextObject->GetClass()->GetPathName() : TEXT("Null");
+	});
+	UE_LOG(LogSlate, Warning, TEXT("DebugLogContextObjectClasses: %s"), *ContextObjectClasses);
+}
+
 void UToolMenuEntryExtensions::SetLabel(UPARAM(ref) FToolMenuEntry& Target, const FText& Label)
 {
 	Target.Label = Label;
diff --git a/Engine/Source/Developer/ToolMenus/Private/ToolMenusBlueprintLibrary.h b/Engine/Source/Developer/ToolMenus/Private/ToolMenusBlueprintLibrary.h
--- a/Engine/Source/Developer/ToolMenus/Private/ToolMenusBlueprintLibrary.h
+++ b/Engine/Source/Developer/ToolMenus/Private/ToolMenusBlueprintLibrary.h
@@ -23,6 +23,9 @@ public:
 
 	UFUNCTION(BlueprintCallable, Category = "Tool Menus", meta = (ScriptMethod, DeterminesOutputType = "InClass"))
 	static UObject* FindByClass(const FToolMenuContext& Context, TSubclassOf<UObject> InClass);
+
+	UFUNCTION(BlueprintCallable, Category = "Tool Menus", meta = (ScriptMethod))
+	static void DebugLogContextObjectClasses(const FToolMenuContext& Context);
 };
 
 UCLASS()
diff --git a/Engine/Source/Developer/ToolMenus/Public/ToolMenuContext.h b/Engine/Source/Developer/ToolMenus/Public/ToolMenuContext.h
--- a/Engine/Source/Developer/ToolMenus/Public/ToolMenuContext.h
+++ b/Engine/Source/Developer/ToolMenus/Public/ToolMenuContext.h
@@ -85,6 +85,7 @@ public:
 	friend class UToolMenus;
 	friend class UToolMenu;
 	friend struct FToolMenuEntry;
+	friend class UToolMenuContextExtensions;
 
 	bool IsEditing() const { return bIsEditing; }
 	void SetIsEditing(bool InIsEditing) { bIsEditing = InIsEditing; }
