Remove constructor of FToolMenuEntryToolBarData and move default values to fields

diff --git a/Engine/Source/Developer/ToolMenus/Public/ToolMenuEntry.h b/Engine/Source/Developer/ToolMenus/Public/ToolMenuEntry.h
--- a/Engine/Source/Developer/ToolMenus/Public/ToolMenuEntry.h
+++ b/Engine/Source/Developer/ToolMenus/Public/ToolMenuEntry.h
@@ -44,13 +44,6 @@ struct FToolMenuEntryOptionsDropdownData
 struct FToolMenuEntryToolBarData
 {
 public:
-	FToolMenuEntryToolBarData() :
-		bSimpleComboBox(false),
-		bIsFocusable(false),
-		bForceSmallIcons(false)
-	{
-	}
-
 	/** Optional override label to use when the entry appears in a toolbar. */
 	TAttribute<FText> LabelOverride;
 
@@ -62,13 +55,13 @@ public:
 
 	TSharedPtr<FToolMenuEntryOptionsDropdownData> OptionsDropdownData;
 
-	bool bSimpleComboBox;
+	bool bSimpleComboBox = false;
 
 	/** Whether ToolBar will have Focusable buttons. */
-	bool bIsFocusable;
+	bool bIsFocusable = false;
 
 	/** Whether this toolbar should always use small icons, regardless of the current settings. */
-	bool bForceSmallIcons;
+	bool bForceSmallIcons = false;
 
 	/** Optional override placement to customize menu placement once opened via e.g. a toolbar menu button */
 	TAttribute<EMenuPlacement> PlacementOverride;
