[ToolMenus] Forwarded ResizeParams to ToolBarBuilder.AddWidget

#rnx

diff --git a/Engine/Source/Developer/ToolMenus/Private/ToolMenus.cpp b/Engine/Source/Developer/ToolMenus/Private/ToolMenus.cpp
--- a/Engine/Source/Developer/ToolMenus/Private/ToolMenus.cpp
+++ b/Engine/Source/Developer/ToolMenus/Private/ToolMenus.cpp
@@ -1924,7 +1924,7 @@ void UToolMenus::PopulateToolBarBuilderWithEntry(
 			StyleParams.VerticalAlignment = VAlign_Fill;
 		}
 
-		ToolBarBuilder.AddWidget(Widget.ToSharedRef(), StyleParams, Block.TutorialHighlightName, Block.WidgetData.bSearchable, FNewMenuDelegate(), VisibilityOverride);
+		ToolBarBuilder.AddWidget(Widget.ToSharedRef(), StyleParams, Block.TutorialHighlightName, Block.WidgetData.bSearchable, FNewMenuDelegate(), VisibilityOverride, ActualResizeParams);
 	}
 	else
 	{
