Expose FToolMenuSection::Alignment to scripts via UToolMenu::AddSectionScript

Use from Python like this:

toolmenus = unreal.ToolMenus.get()
menu = toolmenus.extend_menu("MyMenu")
menu.add_section("SectionName", alignment = unreal.ToolMenuSectionAlign.MIDDLE)

diff --git a/Engine/Source/Developer/ToolMenus/Private/ToolMenu.cpp b/Engine/Source/Developer/ToolMenus/Private/ToolMenu.cpp
--- a/Engine/Source/Developer/ToolMenus/Private/ToolMenu.cpp
+++ b/Engine/Source/Developer/ToolMenus/Private/ToolMenu.cpp
@@ -266,7 +266,13 @@ FToolMenuSection& UToolMenu::AddSection(const FName SectionName, const TAttribut
 	return NewSection;
 }
 
-void UToolMenu::AddSectionScript(const FName SectionName, const FText& InLabel, const FName InsertName, const EToolMenuInsertType InsertType)
+void UToolMenu::AddSectionScript(
+	const FName SectionName,
+	const FText& InLabel,
+	const FName InsertName,
+	const EToolMenuInsertType InsertType,
+	EToolMenuSectionAlign Alignment
+)
 {
 	FToolMenuSection& Section = FindOrAddSection(SectionName);
 
@@ -279,6 +285,8 @@ void UToolMenu::AddSectionScript(const FName SectionName, const FText& InLabel,
 	{
 		Section.InsertPosition = FToolMenuInsert(InsertName, InsertType);
 	}
+
+	Section.Alignment = Alignment;
 }
 
 void UToolMenu::AddDynamicSectionScript(const FName SectionName, UToolMenuSectionDynamic* InObject)
diff --git a/Engine/Source/Developer/ToolMenus/Public/ToolMenu.h b/Engine/Source/Developer/ToolMenus/Public/ToolMenu.h
--- a/Engine/Source/Developer/ToolMenus/Public/ToolMenu.h
+++ b/Engine/Source/Developer/ToolMenus/Public/ToolMenu.h
@@ -31,7 +31,13 @@ public:
 	void InitMenu(const FToolMenuOwner Owner, FName Name, FName Parent = NAME_None, EMultiBoxType Type = EMultiBoxType::Menu);
 
 	UFUNCTION(BlueprintCallable, Category = "Tool Menus", meta = ( ScriptName = "AddSection", DisplayName = "Add Section", AutoCreateRefTerm = "Label", AdvancedDisplay = "InsertName,InsertType" ))
-	void AddSectionScript(const FName SectionName, const FText& Label = FText(), const FName InsertName = NAME_None, const EToolMenuInsertType InsertType = EToolMenuInsertType::Default);
+	void AddSectionScript(
+		const FName SectionName,
+		const FText& Label = FText(),
+		const FName InsertName = NAME_None,
+		const EToolMenuInsertType InsertType = EToolMenuInsertType::Default,
+		const EToolMenuSectionAlign Alignment = EToolMenuSectionAlign::Default
+	);
 
 	UFUNCTION(BlueprintCallable, Category = "Tool Menus", meta = ( ScriptName = "AddDynamicSection", DisplayName = "Add Dynamic Section" ))
 	void AddDynamicSectionScript(const FName SectionName, UToolMenuSectionDynamic* Object);
