Expose FToolMenuEntry::ShowInToolbarTopLevel to scripts

Use from Python by overriding unreal.ToolMenuEntryScript.show_in_toolbar_top_level.

diff --git a/Engine/Source/Developer/ToolMenus/Private/ToolMenuEntryScript.cpp b/Engine/Source/Developer/ToolMenus/Private/ToolMenuEntryScript.cpp
--- a/Engine/Source/Developer/ToolMenus/Private/ToolMenuEntryScript.cpp
+++ b/Engine/Source/Developer/ToolMenus/Private/ToolMenuEntryScript.cpp
@@ -110,6 +110,27 @@ TAttribute<FSlateIcon> UToolMenuEntryScript::CreateIconAttribute(FToolMenuContex
 	return Data.Icon.GetSlateIcon();
 }
 
+TAttribute<bool> UToolMenuEntryScript::CreateShowInToolbarTopLevelAttribute(FToolMenuContext& Context)
+{
+	static const FName FunctionName = GET_FUNCTION_NAME_CHECKED(UToolMenuEntryScript, ShowInToolbarTopLevel);
+	if (GetClass()->IsFunctionImplementedInScript(FunctionName))
+	{
+		TWeakObjectPtr<UToolMenuEntryScript> WeakScriptObject(this);
+		TAttribute<bool>::FGetter Getter;
+		Getter.BindLambda(
+			[WeakScriptObject, Context]()
+			{
+				UToolMenuEntryScript* Object = GetIfCanSafelyRouteCall(WeakScriptObject);
+				return Object ? Object->ShowInToolbarTopLevel(Context) : false;
+			}
+		);
+
+		return TAttribute<bool>::Create(Getter);
+	}
+
+	return TAttribute<bool>();
+}
+
 FSlateIcon UToolMenuEntryScript::GetSlateIcon(const FToolMenuContext& Context) const
 {
 	return GetIcon(Context).GetSlateIcon();
diff --git a/Engine/Source/Developer/ToolMenus/Private/ToolMenus.cpp b/Engine/Source/Developer/ToolMenus/Private/ToolMenus.cpp
--- a/Engine/Source/Developer/ToolMenus/Private/ToolMenus.cpp
+++ b/Engine/Source/Developer/ToolMenus/Private/ToolMenus.cpp
@@ -1490,9 +1490,15 @@ void UToolMenus::PopulateToolBarBuilderWithTopLevelChildren(
 			break;
 		}
 
-		const bool bIsBound = Block->ShowInToolbarTopLevel.IsBound();
-		const bool bIsSetToValue = !bIsBound && Block->ShowInToolbarTopLevel.IsSet();
-		const bool bIsSetToTrueValue = bIsSetToValue && Block->ShowInToolbarTopLevel.Get();
+		const TAttribute<bool> ScriptShowInToolbarTopLevel =
+			Block->ScriptObject ? Block->ScriptObject->CreateShowInToolbarTopLevelAttribute(ParentMenu->Context)
+								: TAttribute<bool>();
+
+		const bool bIsBound = Block->ShowInToolbarTopLevel.IsBound() || ScriptShowInToolbarTopLevel.IsBound();
+		const bool bIsSetToValue = !bIsBound
+								&& (Block->ShowInToolbarTopLevel.IsSet() || ScriptShowInToolbarTopLevel.IsSet());
+		const bool bIsSetToTrueValue = bIsSetToValue
+									&& (Block->ShowInToolbarTopLevel.Get() || ScriptShowInToolbarTopLevel.Get());
 		const bool bIsBoundOrTrue = bIsBound || bIsSetToTrueValue;
 		if (bIsBoundOrTrue || Block->Type == EMultiBlockType::Separator)
 		{
@@ -1798,7 +1804,34 @@ void UToolMenus::PopulateToolBarBuilderWithEntry(
 		else if (Block.ScriptObject)
 		{
 			UToolMenuEntryScript* ScriptObject = Block.ScriptObject;
-			TAttribute<FSlateIcon> Icon = ScriptObject->CreateIconAttribute(MenuData->Context);
+			const TAttribute<FSlateIcon> Icon = ScriptObject->CreateIconAttribute(MenuData->Context);
+
+			TAttribute<EVisibility> ScriptVisibilityOverride = VisibilityOverride;
+
+			// Allow scripts to override ShowInToolbarTopLevel.
+			if (const TAttribute<bool> ScriptShowInToolbarTopLevel =
+					ScriptObject->CreateShowInToolbarTopLevelAttribute(MenuData->Context);
+				bIsRaisingToTopLevel && ScriptShowInToolbarTopLevel.IsSet())
+			{
+				ScriptVisibilityOverride = TAttribute<EVisibility>::CreateLambda(
+					[OriginalVisibilityOverride = VisibilityOverride, ScriptShowInToolbarTopLevel]() -> EVisibility
+					{
+						if (ScriptShowInToolbarTopLevel.IsSet() && !ScriptShowInToolbarTopLevel.Get())
+						{
+							return EVisibility::Collapsed;
+						}
+						else if (OriginalVisibilityOverride.IsSet())
+						{
+							return OriginalVisibilityOverride.Get();
+						}
+						else
+						{
+							return EVisibility::Visible;
+						}
+					}
+				);
+			}
+
 			ToolBarBuilder.AddToolBarButton(
 				UIAction,
 				ScriptObject->Data.Name,
@@ -1807,7 +1840,7 @@ void UToolMenus::PopulateToolBarBuilderWithEntry(
 				Icon,
 				Block.UserInterfaceActionType,
 				Block.TutorialHighlightName,
-				VisibilityOverride,
+				ScriptVisibilityOverride,
 				ToolbarLabelOverride,
 				ActualResizeParams
 			);
diff --git a/Engine/Source/Developer/ToolMenus/Public/ToolMenuEntryScript.h b/Engine/Source/Developer/ToolMenus/Public/ToolMenuEntryScript.h
--- a/Engine/Source/Developer/ToolMenus/Public/ToolMenuEntryScript.h
+++ b/Engine/Source/Developer/ToolMenus/Public/ToolMenuEntryScript.h
@@ -158,6 +158,13 @@ public:
 	UFUNCTION(BlueprintCallable, Category = "Advanced")
 	void InitEntry(const FName OwnerName, const FName Menu, const FName Section, const FName Name, const FText& Label = FText(), const FText& ToolTip = FText());
 
+	UFUNCTION(BlueprintNativeEvent, Category = "Advanced")
+	bool ShowInToolbarTopLevel(const FToolMenuContext& Context) const;
+	virtual bool ShowInToolbarTopLevel_Implementation(const FToolMenuContext& Context) const
+	{
+		return false;
+	}
+
 	FORCEINLINE bool CanSafelyRouteCall() { return !(GIntraFrameDebuggingGameThread || IsUnreachable() || FUObjectThreadContext::Get().IsRoutingPostLoad); }
 
 	static UToolMenuEntryScript* GetIfCanSafelyRouteCall(const TWeakObjectPtr<UToolMenuEntryScript>& InWeak);
@@ -174,6 +181,8 @@ private:
 
 	TAttribute<FSlateIcon> CreateIconAttribute(FToolMenuContext& Context);
 
+	TAttribute<bool> CreateShowInToolbarTopLevelAttribute(FToolMenuContext& Context);
+
 	void ToMenuEntry(FToolMenuEntry& Output);
 
 	bool IsDynamicConstruct() const;
