[ToolMenus] Allow overriding sub-menu style

diff --git a/Engine/Source/Developer/ToolMenus/Private/ToolMenus.cpp b/Engine/Source/Developer/ToolMenus/Private/ToolMenus.cpp
--- a/Engine/Source/Developer/ToolMenus/Private/ToolMenus.cpp
+++ b/Engine/Source/Developer/ToolMenus/Private/ToolMenus.cpp
@@ -135,7 +135,7 @@ public:
 			if (Block.SubMenuData.ConstructMenu.NewToolMenu.IsBound())
 			{
 				// Blocks with no name cannot call PopulateSubMenu()
-				NewMenuDelegate = FNewMenuDelegate::CreateUObject(UToolMenus::Get(), &UToolMenus::PopulateSubMenuWithoutName, TWeakObjectPtr<UToolMenu>(MenuData), Block.SubMenuData.ConstructMenu.NewToolMenu);
+				NewMenuDelegate = FNewMenuDelegate::CreateUObject(UToolMenus::Get(), &UToolMenus::PopulateSubMenuWithoutName, TWeakObjectPtr<UToolMenu>(MenuData), Block);
 			}
 			else
 			{
@@ -172,7 +172,7 @@ public:
 				}
 			}
 
-			NewMenuDelegate = FNewMenuDelegate::CreateUObject(UToolMenus::Get(), &UToolMenus::PopulateSubMenu, TWeakObjectPtr<UToolMenu>(MenuData), BlockNameOverride);
+			NewMenuDelegate = FNewMenuDelegate::CreateUObject(UToolMenus::Get(), &UToolMenus::PopulateSubMenu, TWeakObjectPtr<UToolMenu>(MenuData), Block, BlockNameOverride);
 		}
 
 		if (!bSubMenuAdded)
@@ -546,6 +546,17 @@ TArray<UToolMenu*> UToolMenus::CollectHierarchy(const FName InName, const TMap<F
 	return Result;
 }
 
+void UToolMenus::ApplyStyleToBuilder(FMenuBuilder& InBuilder, const FToolMenuEntryStyle& InStyle)
+{
+	const ISlateStyle* StyleSet = InBuilder.GetStyleSet();
+	if (!StyleSet && InStyle.StyleSet.IsSet())
+	{
+		StyleSet = InStyle.StyleSet.GetValue();
+	}
+
+	InBuilder.SetStyle(StyleSet, InStyle.StyleName.Get(InBuilder.GetStyleName()));
+}
+
 TArray<UToolMenu*> UToolMenus::CollectHierarchy(const FName InName)
 {
 	TMap<FName, FName> UnregisteredParents;
@@ -1316,16 +1327,20 @@ UToolMenu* UToolMenus::GenerateSubMenu(const UToolMenu* InGeneratedParent, const
 	return nullptr;
 }
 
-void UToolMenus::PopulateSubMenu(FMenuBuilder& MenuBuilder, TWeakObjectPtr<UToolMenu> InParent, const FName InBlockName)
+void UToolMenus::PopulateSubMenu(FMenuBuilder& MenuBuilder, TWeakObjectPtr<UToolMenu> InParent, const FToolMenuEntry InEntry, const FName InBlockName)
 {
 	if (UToolMenu* GeneratedMenu = GenerateSubMenu(InParent.Get(), InBlockName))
 	{
 		MenuBuilder.SetExtendersEnabled(GeneratedMenu->bExtendersEnabled);
+
+		// Apply Style override
+		ApplyStyleToBuilder(MenuBuilder, InEntry.SubMenuData.Style);
+
 		PopulateMenuBuilder(MenuBuilder, GeneratedMenu);
 	}
 }
 
-void UToolMenus::PopulateSubMenuWithoutName(FMenuBuilder& MenuBuilder, TWeakObjectPtr<UToolMenu> InParent, const FNewToolMenuDelegate InNewToolMenuDelegate)
+void UToolMenus::PopulateSubMenuWithoutName(FMenuBuilder& MenuBuilder, TWeakObjectPtr<UToolMenu> InParent, const FToolMenuEntry InEntry)
 {
 	const UToolMenu* InGeneratedParent = InParent.Get();
 	if (InGeneratedParent == nullptr)
@@ -1333,7 +1348,7 @@ void UToolMenus::PopulateSubMenuWithoutName(FMenuBuilder& MenuBuilder, TWeakObje
 		return;
 	}
 
-	if (InNewToolMenuDelegate.IsBound())
+	if (InEntry.SubMenuData.ConstructMenu.NewToolMenu.IsBound())
 	{
 		UToolMenu* Menu = NewToolMenuObject(FName(TEXT("SubMenuWithoutName")), NAME_None); // Menu does not have a name
 		Menu->Context = InGeneratedParent->Context;
@@ -1342,7 +1357,10 @@ void UToolMenus::PopulateSubMenuWithoutName(FMenuBuilder& MenuBuilder, TWeakObje
 		Menu->SubMenuParent = InGeneratedParent;
 		Menu->SubMenuSourceEntryName = NAME_None; // Entry does not have a name
 
-		InNewToolMenuDelegate.Execute(Menu);
+		// Apply Style override
+		ApplyStyleToBuilder(MenuBuilder, InEntry.SubMenuData.Style);
+
+		InEntry.SubMenuData.ConstructMenu.NewToolMenu.Execute(Menu);
 		Menu->MenuName = NAME_None; // Menu does not have a name
 
 		PopulateMenuBuilder(MenuBuilder, Menu);
@@ -2119,7 +2137,7 @@ void UToolMenus::PopulateMenuBarBuilder(FMenuBarBuilder& MenuBarBuilder, UToolMe
 				MenuBarBuilder.AddPullDownMenu(
 					Block.Label,
 					Block.ToolTip,
-					FNewMenuDelegate::CreateUObject(this, &UToolMenus::PopulateSubMenu, TWeakObjectPtr<UToolMenu>(MenuData), Block.Name),
+					FNewMenuDelegate::CreateUObject(this, &UToolMenus::PopulateSubMenu, TWeakObjectPtr<UToolMenu>(MenuData), Block, Block.Name),
 					Block.Name
 				);
 			}
diff --git a/Engine/Source/Developer/ToolMenus/Public/ToolMenuEntry.h b/Engine/Source/Developer/ToolMenus/Public/ToolMenuEntry.h
--- a/Engine/Source/Developer/ToolMenus/Public/ToolMenuEntry.h
+++ b/Engine/Source/Developer/ToolMenus/Public/ToolMenuEntry.h
@@ -17,6 +17,16 @@
 class UToolMenuEntryScript;
 struct FKeyEvent;
 
+/** A StyleSet, StyleName pair */
+struct FToolMenuEntryStyle
+{
+	/** The StyleSet used to create the calling widget. Default is auto-calculated or derived from elsewhere */
+	TOptional<const class ISlateStyle*> StyleSet = nullptr;
+
+	/** The name of the style (within the StyleSet) to use to create the calling widget. Default is auto-calculated or derived from elsewhere */
+	TOptional<FName> StyleName;
+};
+
 struct FToolMenuEntrySubMenuData
 {
 public:
@@ -32,6 +42,9 @@ public:
 	/** Entry placed into the parent's menu when there is only one entry */
 	bool bAutoCollapse;
 	FNewToolMenuChoice ConstructMenu;
+
+	/** Optionally specify to override the default styling */
+	FToolMenuEntryStyle Style;
 };
 
 struct FToolMenuEntryOptionsDropdownData
diff --git a/Engine/Source/Developer/ToolMenus/Public/ToolMenus.h b/Engine/Source/Developer/ToolMenus/Public/ToolMenus.h
--- a/Engine/Source/Developer/ToolMenus/Public/ToolMenus.h
+++ b/Engine/Source/Developer/ToolMenus/Public/ToolMenus.h
@@ -418,10 +418,13 @@ private:
 
 	static void ExecuteStringCommand(const FToolMenuStringCommand StringCommand, const FToolMenuContext Context);
 
-	void PopulateSubMenu(FMenuBuilder& Builder, TWeakObjectPtr<UToolMenu> InParent, const FName InBlockName);
-	void PopulateSubMenuWithoutName(FMenuBuilder& MenuBuilder, TWeakObjectPtr<UToolMenu> InParent, const FNewToolMenuDelegate InNewToolMenuDelegate);
+	void PopulateSubMenu(FMenuBuilder& Builder, TWeakObjectPtr<UToolMenu> InParent, const FToolMenuEntry InEntry, const FName InBlockName);
+	void PopulateSubMenuWithoutName(FMenuBuilder& MenuBuilder, TWeakObjectPtr<UToolMenu> InParent, const FToolMenuEntry InEntry);
 	TArray<UToolMenu*> CollectHierarchy(const FName Name, const TMap<FName, FName>& UnregisteredParentNames);
 
+	/** Applies the style to the given Builder, as specified by the provided Style. */
+	void ApplyStyleToBuilder(FMenuBuilder& InBuilder, const FToolMenuEntryStyle& InStyle);
+
 	void ListAllParents(const FName Name, TArray<FName>& AllParents);
 
 	void AssembleMenu(UToolMenu* GeneratedMenu, const UToolMenu* Other);
