Fix editor crash refreshing widgets of tool menus

diff --git a/Engine/Source/Developer/ToolMenus/Private/ToolMenus.cpp b/Engine/Source/Developer/ToolMenus/Private/ToolMenus.cpp
--- a/Engine/Source/Developer/ToolMenus/Private/ToolMenus.cpp
+++ b/Engine/Source/Developer/ToolMenus/Private/ToolMenus.cpp
@@ -2477,13 +2477,13 @@ void UToolMenus::AddReferencedObjects(UObject* InThis, FReferenceCollector& Coll
 
 	for (auto WidgetsForMenuNameIt = This->GeneratedMenuWidgets.CreateIterator(); WidgetsForMenuNameIt; ++WidgetsForMenuNameIt)
 	{
-		FGeneratedToolMenuWidgets& WidgetsForMenuName = WidgetsForMenuNameIt->Value;
+		TSharedPtr<FGeneratedToolMenuWidgets>& WidgetsForMenuName = WidgetsForMenuNameIt->Value;
 
-		for (auto Instance = WidgetsForMenuName.Instances.CreateIterator(); Instance; ++Instance)
+		for (auto Instance = WidgetsForMenuName->Instances.CreateIterator(); Instance; ++Instance)
 		{
-			if (Instance->Widget.IsValid())
+			if (Instance->Get()->Widget.IsValid())
 			{
-				Collector.AddReferencedObject(Instance->GeneratedMenu, InThis);
+				Collector.AddReferencedObject(Instance->Get()->GeneratedMenu, InThis);
 			}
 			else
 			{
@@ -2491,7 +2491,7 @@ void UToolMenus::AddReferencedObjects(UObject* InThis, FReferenceCollector& Coll
 			}
 		}
 
-		if (WidgetsForMenuName.Instances.Num() == 0)
+		if (WidgetsForMenuName->Instances.Num() == 0)
 		{
 			WidgetsForMenuNameIt.RemoveCurrent();
 		}
@@ -2546,17 +2546,15 @@ void UToolMenus::RemoveSubstitutionDuringGenerate(const FName InMenu)
 		const FName OverrideMenuName = *FoundOverrideMenuName;
 
 		// Update all active widget instances of this menu
-		FGeneratedToolMenuWidgets* OverrideMenuWidgets = GeneratedMenuWidgets.Find(OverrideMenuName);
-		if (OverrideMenuWidgets)
+		if (TSharedPtr<FGeneratedToolMenuWidgets> OverrideMenuWidgets = GeneratedMenuWidgets.FindRef(OverrideMenuName))
 		{
-			FGeneratedToolMenuWidgets* DestMenuWidgets = GeneratedMenuWidgets.Find(InMenu);
-			if (DestMenuWidgets)
+			if (TSharedPtr<FGeneratedToolMenuWidgets> DestMenuWidgets = GeneratedMenuWidgets.FindRef(InMenu))
 			{
 				DestMenuWidgets->Instances.Append(OverrideMenuWidgets->Instances);
 			}
 			else
 			{
-				GeneratedMenuWidgets.Add(InMenu, *OverrideMenuWidgets);
+				GeneratedMenuWidgets.Add(InMenu, OverrideMenuWidgets);
 			}
 
 			GeneratedMenuWidgets.Remove(OverrideMenuName);
@@ -2698,10 +2696,14 @@ TSharedRef<SWidget> UToolMenus::GenerateWidget(UToolMenu* GeneratedMenu)
 		GeneratedWidget = Result;
 	}
 
-	FGeneratedToolMenuWidgets& WidgetsForMenuName = GeneratedMenuWidgets.FindOrAdd(GeneratedMenu->MenuName);
+	TSharedPtr<FGeneratedToolMenuWidgets>& WidgetsForMenuName = GeneratedMenuWidgets.FindOrAdd(GeneratedMenu->MenuName);
+	if (!WidgetsForMenuName.IsValid())
+	{
+		WidgetsForMenuName = MakeShared<FGeneratedToolMenuWidgets>();
+	}
 
 	// Store a copy so that we can call 'Refresh' on menus not in the database
-	FGeneratedToolMenuWidget& GeneratedMenuWidget = WidgetsForMenuName.Instances.AddDefaulted_GetRef();
+	FGeneratedToolMenuWidget& GeneratedMenuWidget = *WidgetsForMenuName->Instances.Add_GetRef(MakeShared<FGeneratedToolMenuWidget>());
 	GeneratedMenuWidget.OriginalMenu = GeneratedMenu;
 	GeneratedMenuWidget.GeneratedMenu = DuplicateObject<UToolMenu>(GeneratedMenu, this, MakeUniqueObjectName(this, UToolMenus::StaticClass(), FName("MenuForRefresh")));
 	GeneratedMenuWidget.GeneratedMenu->bShouldCleanupContextOnDestroy = true;
@@ -2787,14 +2789,17 @@ void UToolMenus::HandleNextTick()
 		{
 			TGuardValue<bool> SuppressRefreshWidgetsRequestsGuard(bSuppressRefreshWidgetsRequests, true);
 
-			for (auto WidgetsForMenuNameIt = GeneratedMenuWidgets.CreateIterator(); WidgetsForMenuNameIt; ++WidgetsForMenuNameIt)
+			// Copy before enumerate because modified inside RefreshMenuWidget
+			TMap<FName, TSharedPtr<FGeneratedToolMenuWidgets>> GeneratedMenuWidgetsCopy = GeneratedMenuWidgets;
+			for (TPair<FName, TSharedPtr<FGeneratedToolMenuWidgets>>& WidgetsForMenuNameIt : GeneratedMenuWidgetsCopy)
 			{
-				FGeneratedToolMenuWidgets& WidgetsForMenuName = WidgetsForMenuNameIt->Value;
-				for (auto Instance = WidgetsForMenuName.Instances.CreateIterator(); Instance; ++Instance)
+				// Copy before enumerate because modified inside RefreshMenuWidget
+				TArray<TSharedPtr<FGeneratedToolMenuWidget>> InstancesCopy = WidgetsForMenuNameIt.Value->Instances;
+				for (TSharedPtr<FGeneratedToolMenuWidget>& Instance : InstancesCopy)
 				{
 					if (Instance->Widget.IsValid())
 					{
-						RefreshMenuWidget(WidgetsForMenuNameIt->Key, *Instance);
+						RefreshMenuWidget(WidgetsForMenuNameIt.Key, *Instance);
 					}
 				}
 			}
@@ -2811,18 +2816,18 @@ void UToolMenus::CleanupStaleWidgets()
 	bool bModified = false;
 	for (auto WidgetsForMenuNameIt = GeneratedMenuWidgets.CreateIterator(); WidgetsForMenuNameIt; ++WidgetsForMenuNameIt)
 	{
-		FGeneratedToolMenuWidgets& WidgetsForMenuName = WidgetsForMenuNameIt->Value;
+		TSharedPtr<FGeneratedToolMenuWidgets>& WidgetsForMenuName = WidgetsForMenuNameIt->Value;
 
-		for (auto Instance = WidgetsForMenuName.Instances.CreateIterator(); Instance; ++Instance)
+		for (auto Instance = WidgetsForMenuName->Instances.CreateIterator(); Instance; ++Instance)
 		{
-			if (!Instance->Widget.IsValid())
+			if (!(*Instance)->Widget.IsValid())
 			{
 				bModified = true;
 				Instance.RemoveCurrent();
 			}
 		}
 
-		if (WidgetsForMenuName.Instances.Num() == 0)
+		if (WidgetsForMenuName->Instances.Num() == 0)
 		{
 			bModified = true;
 			WidgetsForMenuNameIt.RemoveCurrent();
@@ -2838,9 +2843,13 @@ void UToolMenus::CleanupStaleWidgets()
 bool UToolMenus::RefreshMenuWidget(const FName InName)
 {
 	bool bRefreshedAnyWidget = false;
-	if (FGeneratedToolMenuWidgets* WidgetsForMenuName = GeneratedMenuWidgets.Find(InName))
+
+	// Copy the TSharedPtr because RefreshMenuWidget can modify GeneratedMenuWidgets and delete any memory we point to
+	if (TSharedPtr<FGeneratedToolMenuWidgets> WidgetsForMenuName = GeneratedMenuWidgets.FindRef(InName); WidgetsForMenuName.IsValid())
 	{
-		for (auto Instance = WidgetsForMenuName->Instances.CreateIterator(); Instance; ++Instance)
+		// Copy before enumerate because entries can be added to this array inside RefreshMenuWidget
+		TArray<TSharedPtr<FGeneratedToolMenuWidget>> InstancesCopy = WidgetsForMenuName->Instances;
+		for (TSharedPtr<FGeneratedToolMenuWidget>& Instance : InstancesCopy)
 		{
 			if (RefreshMenuWidget(InName, *Instance))
 			{
@@ -2848,7 +2857,8 @@ bool UToolMenus::RefreshMenuWidget(const FName InName)
 			}
 			else
 			{
-				Instance.RemoveCurrent();
+				// Remove from original instead of InstancesCopy
+				WidgetsForMenuName->Instances.Remove(Instance);
 			}
 		}
 	}
diff --git a/Engine/Source/Developer/ToolMenus/Public/ToolMenus.h b/Engine/Source/Developer/ToolMenus/Public/ToolMenus.h
--- a/Engine/Source/Developer/ToolMenus/Public/ToolMenus.h
+++ b/Engine/Source/Developer/ToolMenus/Public/ToolMenus.h
@@ -50,7 +50,7 @@ struct FGeneratedToolMenuWidget
 
 struct FGeneratedToolMenuWidgets
 {
-	TArray<FGeneratedToolMenuWidget> Instances;
+	TArray<TSharedPtr<FGeneratedToolMenuWidget>> Instances;
 };
 
 /*
@@ -458,7 +458,7 @@ private:
 	UPROPERTY()
 	TMap<FName, TObjectPtr<UToolMenu>> Menus;
 
-	TMap<FName, FGeneratedToolMenuWidgets> GeneratedMenuWidgets;
+	TMap<FName, TSharedPtr<FGeneratedToolMenuWidgets>> GeneratedMenuWidgets;
 
 	TMap<TWeakPtr<FMultiBox>, TArray<TObjectPtr<const UObject>>> WidgetObjectReferences;
 
