Fix crash in tool menus now that SlateStyles are being deleted.
#rnx

diff --git a/Engine/Source/Developer/ToolMenus/Private/ToolMenu.cpp b/Engine/Source/Developer/ToolMenus/Private/ToolMenu.cpp
--- a/Engine/Source/Developer/ToolMenus/Private/ToolMenu.cpp
+++ b/Engine/Source/Developer/ToolMenus/Private/ToolMenu.cpp
@@ -9,6 +9,7 @@
 #include "Framework/MultiBox/MultiBox.h"
 #include "Framework/MultiBox/MultiBoxBuilder.h"
 #include "Internationalization/Internationalization.h"
+#include "Styling/SlateStyleRegistry.h"
 
 #include UE_INLINE_GENERATED_CPP_BY_NAME(ToolMenu)
 
@@ -23,7 +24,6 @@ UToolMenu::UToolMenu() :
 	, bRegistered(false)
 	, bIsRegistering(false)
 	, bExtendersEnabled(true)
-	, StyleSet(&FCoreStyle::Get())
 	, MaxHeight(1000.f)
 {
 }
@@ -60,14 +60,15 @@ FReply UToolMenu::OnKeyDown(const FGeometry& MyGeometry, const FKeyEvent& InKeyE
 
 const ISlateStyle* UToolMenu::GetStyleSet() const
 {
-	return StyleSet;
+	return !StyleSetName.IsNone() ? FSlateStyleRegistry::FindSlateStyle(StyleSetName) : &FCoreStyle::Get();
 }
 
 void UToolMenu::SetStyleSet(const ISlateStyle* InStyleSet)
 {
-	if (InStyleSet && InStyleSet != StyleSet)
+	const FName NewStyleSetName = InStyleSet ? InStyleSet->GetStyleSetName() : NAME_None;
+	if (NewStyleSetName != StyleName)
 	{
-		StyleSet = InStyleSet;
+		StyleSetName = NewStyleSetName;
 	}
 }
 
@@ -80,7 +81,7 @@ void UToolMenu::InitGeneratedCopy(const UToolMenu* Source, const FName InMenuNam
 	StyleName = Source->StyleName;
 	TutorialHighlightName = Source->TutorialHighlightName;
 	MenuType = Source->MenuType;
-	StyleSet = Source->StyleSet;
+	StyleSetName = Source->StyleSetName;
 	bShouldCloseWindowAfterMenuSelection = Source->bShouldCloseWindowAfterMenuSelection;
 	bCloseSelfOnly = Source->bCloseSelfOnly;
 	bSearchable = Source->bSearchable;
diff --git a/Engine/Source/Developer/ToolMenus/Private/ToolMenus.cpp b/Engine/Source/Developer/ToolMenus/Private/ToolMenus.cpp
--- a/Engine/Source/Developer/ToolMenus/Private/ToolMenus.cpp
+++ b/Engine/Source/Developer/ToolMenus/Private/ToolMenus.cpp
@@ -1691,7 +1691,7 @@ void UToolMenus::PopulateToolBarBuilderWithTopLevelChildren(
 		// Never show spacers in overflow menus.
 		ResizeParams.VisibleInOverflow = false;
 
-		const FToolBarStyle ToolbarStyle = ParentMenu->StyleSet->GetWidgetStyle<FToolBarStyle>(ParentMenu->StyleName);
+		const FToolBarStyle ToolbarStyle = ParentMenu->GetStyleSet()->GetWidgetStyle<FToolBarStyle>(ParentMenu->StyleName);
 		ToolBarBuilder.AddWidget(
 			SNew(SSpacer).Size(FVector2D(ToolbarStyle.RaisedChildrenRightPadding, 0)),
 			StyleParams,
@@ -2640,12 +2640,12 @@ TSharedRef<SWidget> UToolMenus::GenerateWidget(UToolMenu* GeneratedMenu)
 			}
 		}
 
-		FMenuBuilder MenuBuilder(GeneratedMenu->bShouldCloseWindowAfterMenuSelection, GeneratedMenu->Context.CommandList, GeneratedMenu->Context.GetAllExtenders(), GeneratedMenu->bCloseSelfOnly, GeneratedMenu->StyleSet, GeneratedMenu->bSearchable, GeneratedMenu->MenuName);
+		FMenuBuilder MenuBuilder(GeneratedMenu->bShouldCloseWindowAfterMenuSelection, GeneratedMenu->Context.CommandList, GeneratedMenu->Context.GetAllExtenders(), GeneratedMenu->bCloseSelfOnly, GeneratedMenu->GetStyleSet(), GeneratedMenu->bSearchable, GeneratedMenu->MenuName);
 
 		// Default consistent style is applied, necessary for toolbars to be displayed as menus
 		//if (GeneratedMenu->StyleName != NAME_None)
 		//{
-		//	MenuBuilder.SetStyle(GeneratedMenu->StyleSet, GeneratedMenu->StyleName);
+		//	MenuBuilder.SetStyle(GeneratedMenu->GetStyleSet(), GeneratedMenu->StyleName);
 		//}
 
 		MenuBuilder.SetExtendersEnabled(GeneratedMenu->bExtendersEnabled);
@@ -2659,11 +2659,11 @@ TSharedRef<SWidget> UToolMenus::GenerateWidget(UToolMenu* GeneratedMenu)
 	}
 	else if (GeneratedMenu->MenuType == EMultiBoxType::Menu)
 	{
-		FMenuBuilder MenuBuilder(GeneratedMenu->bShouldCloseWindowAfterMenuSelection, GeneratedMenu->Context.CommandList, GeneratedMenu->Context.GetAllExtenders(), GeneratedMenu->bCloseSelfOnly, GeneratedMenu->StyleSet, GeneratedMenu->bSearchable, GeneratedMenu->MenuName);
+		FMenuBuilder MenuBuilder(GeneratedMenu->bShouldCloseWindowAfterMenuSelection, GeneratedMenu->Context.CommandList, GeneratedMenu->Context.GetAllExtenders(), GeneratedMenu->bCloseSelfOnly, GeneratedMenu->GetStyleSet(), GeneratedMenu->bSearchable, GeneratedMenu->MenuName);
 
 		if (GeneratedMenu->StyleName != NAME_None)
 		{
-			MenuBuilder.SetStyle(GeneratedMenu->StyleSet, GeneratedMenu->StyleName);
+			MenuBuilder.SetStyle(GeneratedMenu->GetStyleSet(), GeneratedMenu->StyleName);
 		}
 
 		MenuBuilder.SetExtendersEnabled(GeneratedMenu->bExtendersEnabled);
@@ -2673,11 +2673,11 @@ TSharedRef<SWidget> UToolMenus::GenerateWidget(UToolMenu* GeneratedMenu)
 	}
 	else if (GeneratedMenu->MenuType == EMultiBoxType::MenuBar)
 	{
-		FMenuBarBuilder MenuBarBuilder(GeneratedMenu->Context.CommandList, GeneratedMenu->Context.GetAllExtenders(), GeneratedMenu->StyleSet, GeneratedMenu->MenuName);
+		FMenuBarBuilder MenuBarBuilder(GeneratedMenu->Context.CommandList, GeneratedMenu->Context.GetAllExtenders(), GeneratedMenu->GetStyleSet(), GeneratedMenu->MenuName);
 
 		if (GeneratedMenu->StyleName != NAME_None)
 		{
-			MenuBarBuilder.SetStyle(GeneratedMenu->StyleSet, GeneratedMenu->StyleName);
+			MenuBarBuilder.SetStyle(GeneratedMenu->GetStyleSet(), GeneratedMenu->StyleName);
 		}
 
 		MenuBarBuilder.SetExtendersEnabled(GeneratedMenu->bExtendersEnabled);
@@ -2694,7 +2694,7 @@ TSharedRef<SWidget> UToolMenus::GenerateWidget(UToolMenu* GeneratedMenu)
 
 		if (GeneratedMenu->StyleName != NAME_None)
 		{
-			ToolbarBuilder.SetStyle(GeneratedMenu->StyleSet, GeneratedMenu->StyleName);
+			ToolbarBuilder.SetStyle(GeneratedMenu->GetStyleSet(), GeneratedMenu->StyleName);
 		}
 
 		PopulateToolBarBuilder(ToolbarBuilder, GeneratedMenu);
@@ -2715,7 +2715,7 @@ TSharedRef<SWidget> UToolMenus::GenerateWidget(UToolMenu* GeneratedMenu)
 	GeneratedMenuWidget.GeneratedMenu->bShouldCleanupContextOnDestroy = true;
 	// Copy native properties that serialize does not
 	GeneratedMenuWidget.GeneratedMenu->Context = GeneratedMenu->Context;
-	GeneratedMenuWidget.GeneratedMenu->StyleSet = GeneratedMenu->StyleSet;
+	GeneratedMenuWidget.GeneratedMenu->StyleSetName = GeneratedMenu->StyleSetName;
 	GeneratedMenuWidget.GeneratedMenu->StyleName = GeneratedMenu->StyleName;
 
 	if (GeneratedWidget)
@@ -2895,12 +2895,12 @@ bool UToolMenus::RefreshMenuWidget(const FName InName, FGeneratedToolMenuWidget&
 	TSharedRef<SMultiBoxWidget> MultiBoxWidget = StaticCastSharedRef<SMultiBoxWidget>(GeneratedMenuWidget.Widget.Pin().ToSharedRef());
 	if (GeneratedMenu->MenuType == EMultiBoxType::Menu)
 	{
-		FMenuBuilder MenuBuilder(GeneratedMenu->bShouldCloseWindowAfterMenuSelection, GeneratedMenu->Context.CommandList, GeneratedMenu->Context.GetAllExtenders(), GeneratedMenu->bCloseSelfOnly, GeneratedMenu->StyleSet, GeneratedMenu->bSearchable);
+		FMenuBuilder MenuBuilder(GeneratedMenu->bShouldCloseWindowAfterMenuSelection, GeneratedMenu->Context.CommandList, GeneratedMenu->Context.GetAllExtenders(), GeneratedMenu->bCloseSelfOnly, GeneratedMenu->GetStyleSet(), GeneratedMenu->bSearchable);
 		MenuBuilder.SetExtendersEnabled(GeneratedMenu->bExtendersEnabled);
 
 		if (GeneratedMenu->StyleName != NAME_None)
 		{
-			MenuBuilder.SetStyle(GeneratedMenu->StyleSet, GeneratedMenu->StyleName);
+			MenuBuilder.SetStyle(GeneratedMenu->GetStyleSet(), GeneratedMenu->StyleName);
 		}
 
 		PopulateMenuBuilder(MenuBuilder, GeneratedMenu);
@@ -2908,12 +2908,12 @@ bool UToolMenus::RefreshMenuWidget(const FName InName, FGeneratedToolMenuWidget&
 	}
 	else if (GeneratedMenu->MenuType == EMultiBoxType::MenuBar)
 	{
-		FMenuBarBuilder MenuBarBuilder(GeneratedMenu->Context.CommandList, GeneratedMenu->Context.GetAllExtenders(), GeneratedMenu->StyleSet);
+		FMenuBarBuilder MenuBarBuilder(GeneratedMenu->Context.CommandList, GeneratedMenu->Context.GetAllExtenders(), GeneratedMenu->GetStyleSet());
 		MenuBarBuilder.SetExtendersEnabled(GeneratedMenu->bExtendersEnabled);
 
 		if (GeneratedMenu->StyleName != NAME_None)
 		{
-			MenuBarBuilder.SetStyle(GeneratedMenu->StyleSet, GeneratedMenu->StyleName);
+			MenuBarBuilder.SetStyle(GeneratedMenu->GetStyleSet(), GeneratedMenu->StyleName);
 		}
 
 		PopulateMenuBarBuilder(MenuBarBuilder, GeneratedMenu);
@@ -2928,7 +2928,7 @@ bool UToolMenus::RefreshMenuWidget(const FName InName, FGeneratedToolMenuWidget&
 
 		if (GeneratedMenu->StyleName != NAME_None)
 		{
-			ToolbarBuilder.SetStyle(GeneratedMenu->StyleSet, GeneratedMenu->StyleName);
+			ToolbarBuilder.SetStyle(GeneratedMenu->GetStyleSet(), GeneratedMenu->StyleName);
 		}
 
 		PopulateToolBarBuilder(ToolbarBuilder, GeneratedMenu);
diff --git a/Engine/Source/Developer/ToolMenus/Public/ToolMenu.h b/Engine/Source/Developer/ToolMenus/Public/ToolMenu.h
--- a/Engine/Source/Developer/ToolMenus/Public/ToolMenu.h
+++ b/Engine/Source/Developer/ToolMenus/Public/ToolMenu.h
@@ -230,7 +230,7 @@ private:
 	bool bIsRegistering;
 	bool bExtendersEnabled;
 
-	const ISlateStyle* StyleSet;
+	FName StyleSetName;
 
 	uint32 MaxHeight;
 };
