ToolMenus - Add ensure when menu widget is being generated and SlateStyleName no longer exists.
Also gracefully substitute style set with CoreStyle to avoid crash when this happens.
#rnx

diff --git a/Engine/Source/Developer/ToolMenus/Private/ToolMenu.cpp b/Engine/Source/Developer/ToolMenus/Private/ToolMenu.cpp
--- a/Engine/Source/Developer/ToolMenus/Private/ToolMenu.cpp
+++ b/Engine/Source/Developer/ToolMenus/Private/ToolMenu.cpp
@@ -60,7 +60,14 @@ FReply UToolMenu::OnKeyDown(const FGeometry& MyGeometry, const FKeyEvent& InKeyE
 
 const ISlateStyle* UToolMenu::GetStyleSet() const
 {
-	return !StyleSetName.IsNone() ? FSlateStyleRegistry::FindSlateStyle(StyleSetName) : &FCoreStyle::Get();
+	if (!StyleSetName.IsNone())
+	{
+		const ISlateStyle* FoundSlateStyleSet = FSlateStyleRegistry::FindSlateStyle(StyleSetName);
+		ensureMsgf(FoundSlateStyleSet, TEXT("Slate Style Set not found: '%s' for menu: '%s'."), *StyleSetName.ToString(), *GetMenuName().ToString());
+		return FoundSlateStyleSet;
+	}
+
+	return &FCoreStyle::Get();
 }
 
 void UToolMenu::SetStyleSet(const ISlateStyle* InStyleSet)
diff --git a/Engine/Source/Developer/ToolMenus/Private/ToolMenus.cpp b/Engine/Source/Developer/ToolMenus/Private/ToolMenus.cpp
--- a/Engine/Source/Developer/ToolMenus/Private/ToolMenus.cpp
+++ b/Engine/Source/Developer/ToolMenus/Private/ToolMenus.cpp
@@ -2625,6 +2625,14 @@ TSharedRef<SWidget> UToolMenus::GenerateWidget(UToolMenu* GeneratedMenu)
 {
 	CleanupStaleWidgetsNextTick();
 
+	const ISlateStyle* StyleSetNotNull = GeneratedMenu->GetStyleSet();
+	const bool bHadStyleSet = StyleSetNotNull != nullptr;
+	if (!bHadStyleSet)
+	{
+		// Avoid crash when style sets are unregistered/deleted, GetStyleSet() will report an ensure when that happens but return null and menu builders crash when passed null StyleSet.
+		StyleSetNotNull = &FCoreStyle::Get();
+	}
+
 	TSharedPtr<SWidget> GeneratedWidget;
 	if (GeneratedMenu->IsEditing())
 	{
@@ -2640,12 +2648,12 @@ TSharedRef<SWidget> UToolMenus::GenerateWidget(UToolMenu* GeneratedMenu)
 			}
 		}
 
-		FMenuBuilder MenuBuilder(GeneratedMenu->bShouldCloseWindowAfterMenuSelection, GeneratedMenu->Context.CommandList, GeneratedMenu->Context.GetAllExtenders(), GeneratedMenu->bCloseSelfOnly, GeneratedMenu->GetStyleSet(), GeneratedMenu->bSearchable, GeneratedMenu->MenuName);
+		FMenuBuilder MenuBuilder(GeneratedMenu->bShouldCloseWindowAfterMenuSelection, GeneratedMenu->Context.CommandList, GeneratedMenu->Context.GetAllExtenders(), GeneratedMenu->bCloseSelfOnly, StyleSetNotNull, GeneratedMenu->bSearchable, GeneratedMenu->MenuName);
 
 		// Default consistent style is applied, necessary for toolbars to be displayed as menus
-		//if (GeneratedMenu->StyleName != NAME_None)
+		//if (bHadStyleSet && GeneratedMenu->StyleName != NAME_None)
 		//{
-		//	MenuBuilder.SetStyle(GeneratedMenu->GetStyleSet(), GeneratedMenu->StyleName);
+		//	MenuBuilder.SetStyle(StyleSetNotNull, GeneratedMenu->StyleName);
 		//}
 
 		MenuBuilder.SetExtendersEnabled(GeneratedMenu->bExtendersEnabled);
@@ -2659,11 +2667,11 @@ TSharedRef<SWidget> UToolMenus::GenerateWidget(UToolMenu* GeneratedMenu)
 	}
 	else if (GeneratedMenu->MenuType == EMultiBoxType::Menu)
 	{
-		FMenuBuilder MenuBuilder(GeneratedMenu->bShouldCloseWindowAfterMenuSelection, GeneratedMenu->Context.CommandList, GeneratedMenu->Context.GetAllExtenders(), GeneratedMenu->bCloseSelfOnly, GeneratedMenu->GetStyleSet(), GeneratedMenu->bSearchable, GeneratedMenu->MenuName);
+		FMenuBuilder MenuBuilder(GeneratedMenu->bShouldCloseWindowAfterMenuSelection, GeneratedMenu->Context.CommandList, GeneratedMenu->Context.GetAllExtenders(), GeneratedMenu->bCloseSelfOnly, StyleSetNotNull, GeneratedMenu->bSearchable, GeneratedMenu->MenuName);
 
-		if (GeneratedMenu->StyleName != NAME_None)
+		if (bHadStyleSet && GeneratedMenu->StyleName != NAME_None)
 		{
-			MenuBuilder.SetStyle(GeneratedMenu->GetStyleSet(), GeneratedMenu->StyleName);
+			MenuBuilder.SetStyle(StyleSetNotNull, GeneratedMenu->StyleName);
 		}
 
 		MenuBuilder.SetExtendersEnabled(GeneratedMenu->bExtendersEnabled);
@@ -2673,11 +2681,11 @@ TSharedRef<SWidget> UToolMenus::GenerateWidget(UToolMenu* GeneratedMenu)
 	}
 	else if (GeneratedMenu->MenuType == EMultiBoxType::MenuBar)
 	{
-		FMenuBarBuilder MenuBarBuilder(GeneratedMenu->Context.CommandList, GeneratedMenu->Context.GetAllExtenders(), GeneratedMenu->GetStyleSet(), GeneratedMenu->MenuName);
+		FMenuBarBuilder MenuBarBuilder(GeneratedMenu->Context.CommandList, GeneratedMenu->Context.GetAllExtenders(), StyleSetNotNull, GeneratedMenu->MenuName);
 
-		if (GeneratedMenu->StyleName != NAME_None)
+		if (bHadStyleSet && GeneratedMenu->StyleName != NAME_None)
 		{
-			MenuBarBuilder.SetStyle(GeneratedMenu->GetStyleSet(), GeneratedMenu->StyleName);
+			MenuBarBuilder.SetStyle(StyleSetNotNull, GeneratedMenu->StyleName);
 		}
 
 		MenuBarBuilder.SetExtendersEnabled(GeneratedMenu->bExtendersEnabled);
@@ -2692,9 +2700,9 @@ TSharedRef<SWidget> UToolMenus::GenerateWidget(UToolMenu* GeneratedMenu)
 		ToolbarBuilder.SetIsFocusable(GeneratedMenu->bToolBarIsFocusable);
 		ToolbarBuilder.SetAllowWrapButton(GeneratedMenu->bAllowToolBarWrapButton);
 
-		if (GeneratedMenu->StyleName != NAME_None)
+		if (bHadStyleSet && GeneratedMenu->StyleName != NAME_None)
 		{
-			ToolbarBuilder.SetStyle(GeneratedMenu->GetStyleSet(), GeneratedMenu->StyleName);
+			ToolbarBuilder.SetStyle(StyleSetNotNull, GeneratedMenu->StyleName);
 		}
 
 		PopulateToolBarBuilder(ToolbarBuilder, GeneratedMenu);
@@ -2891,16 +2899,24 @@ bool UToolMenus::RefreshMenuWidget(const FName InName, FGeneratedToolMenuWidget&
 	UToolMenu* GeneratedMenu = GenerateMenu(InName, GeneratedMenuWidget.GeneratedMenu->Context);
 	GeneratedMenuWidget.GeneratedMenu = GeneratedMenu;
 
+	const ISlateStyle* StyleSetNotNull = GeneratedMenu->GetStyleSet();
+	const bool bHadStyleSet = StyleSetNotNull != nullptr;
+	if (!bHadStyleSet)
+	{
+		// Avoid crash when style sets are unregistered/deleted, GetStyleSet() will report an ensure when that happens but return null and menu builders crash when passed null StyleSet.
+		StyleSetNotNull = &FCoreStyle::Get();
+	}
+
 	// Regenerate Multibox
 	TSharedRef<SMultiBoxWidget> MultiBoxWidget = StaticCastSharedRef<SMultiBoxWidget>(GeneratedMenuWidget.Widget.Pin().ToSharedRef());
 	if (GeneratedMenu->MenuType == EMultiBoxType::Menu)
 	{
-		FMenuBuilder MenuBuilder(GeneratedMenu->bShouldCloseWindowAfterMenuSelection, GeneratedMenu->Context.CommandList, GeneratedMenu->Context.GetAllExtenders(), GeneratedMenu->bCloseSelfOnly, GeneratedMenu->GetStyleSet(), GeneratedMenu->bSearchable);
+		FMenuBuilder MenuBuilder(GeneratedMenu->bShouldCloseWindowAfterMenuSelection, GeneratedMenu->Context.CommandList, GeneratedMenu->Context.GetAllExtenders(), GeneratedMenu->bCloseSelfOnly, StyleSetNotNull, GeneratedMenu->bSearchable);
 		MenuBuilder.SetExtendersEnabled(GeneratedMenu->bExtendersEnabled);
 
-		if (GeneratedMenu->StyleName != NAME_None)
+		if (bHadStyleSet && GeneratedMenu->StyleName != NAME_None)
 		{
-			MenuBuilder.SetStyle(GeneratedMenu->GetStyleSet(), GeneratedMenu->StyleName);
+			MenuBuilder.SetStyle(StyleSetNotNull, GeneratedMenu->StyleName);
 		}
 
 		PopulateMenuBuilder(MenuBuilder, GeneratedMenu);
@@ -2908,12 +2924,12 @@ bool UToolMenus::RefreshMenuWidget(const FName InName, FGeneratedToolMenuWidget&
 	}
 	else if (GeneratedMenu->MenuType == EMultiBoxType::MenuBar)
 	{
-		FMenuBarBuilder MenuBarBuilder(GeneratedMenu->Context.CommandList, GeneratedMenu->Context.GetAllExtenders(), GeneratedMenu->GetStyleSet());
+		FMenuBarBuilder MenuBarBuilder(GeneratedMenu->Context.CommandList, GeneratedMenu->Context.GetAllExtenders(), StyleSetNotNull);
 		MenuBarBuilder.SetExtendersEnabled(GeneratedMenu->bExtendersEnabled);
 
-		if (GeneratedMenu->StyleName != NAME_None)
+		if (bHadStyleSet && GeneratedMenu->StyleName != NAME_None)
 		{
-			MenuBarBuilder.SetStyle(GeneratedMenu->GetStyleSet(), GeneratedMenu->StyleName);
+			MenuBarBuilder.SetStyle(StyleSetNotNull, GeneratedMenu->StyleName);
 		}
 
 		PopulateMenuBarBuilder(MenuBarBuilder, GeneratedMenu);
@@ -2926,9 +2942,9 @@ bool UToolMenus::RefreshMenuWidget(const FName InName, FGeneratedToolMenuWidget&
 		ToolbarBuilder.SetIsFocusable(GeneratedMenu->bToolBarIsFocusable);
 		ToolbarBuilder.SetAllowWrapButton(GeneratedMenu->bAllowToolBarWrapButton);
 
-		if (GeneratedMenu->StyleName != NAME_None)
+		if (bHadStyleSet && GeneratedMenu->StyleName != NAME_None)
 		{
-			ToolbarBuilder.SetStyle(GeneratedMenu->GetStyleSet(), GeneratedMenu->StyleName);
+			ToolbarBuilder.SetStyle(StyleSetNotNull, GeneratedMenu->StyleName);
 		}
 
 		PopulateToolBarBuilder(ToolbarBuilder, GeneratedMenu);
