[ToolMenus] Added various test types that allow testing of an actual menu vs. an expected menu tree

diff --git a/Engine/Source/Developer/ToolMenus/Private/Tests/ToolMenusTestUtilities.cpp b/Engine/Source/Developer/ToolMenus/Private/Tests/ToolMenusTestUtilities.cpp
new file mode 100644
--- /dev/null
+++ b/Engine/Source/Developer/ToolMenus/Private/Tests/ToolMenusTestUtilities.cpp
@@ -0,0 +1,252 @@
+// Copyright Epic Games, Inc. All Rights Reserved.
+
+#include "Tests/ToolMenusTestUtilities.h"
+
+#include "Misc/UObjectTestUtils.h"
+
+namespace UE::ToolMenus
+{
+	FName FMenuEntry::TypeName = "MenuEntry";
+	FName FMenuWildcardEntry::TypeName = "WildcardMenuEntry";
+	FName FMenu::TypeName = "Menu";
+
+	bool FToolMenuAutomationTestAdapter::Matches(const FMenuEntryBase* InExpectedEntry, const FToolMenuEntry& InActualEntry)
+	{
+		// Empty name matches anything
+		if (InExpectedEntry->Name.IsNone())
+		{
+			return true;
+		}
+
+		if (!TestInstance.TestEqual("Entry Name", InActualEntry.Name.ToString(), InExpectedEntry->Name.ToString()))
+		{
+			return false;
+		}
+
+		return true;
+	}
+
+	bool FToolMenuAutomationTestAdapter::Matches(const FMenuEntry& InExpectedEntry, const FToolMenuEntry& InActualEntry)
+	{
+		if (!Matches(static_cast<const FMenuEntryBase*>(&InExpectedEntry), InActualEntry))
+		{
+			return false;
+		}
+
+		if (InExpectedEntry.Type.IsSet() && !UE::CoreUObject::TestEqual(FString::Printf(TEXT("Entry (\"%s\") MultiBlockType"), *InActualEntry.Name.ToString()), InActualEntry.Type, InExpectedEntry.Type.GetValue(), TestInstance))
+		{
+			return false;
+		}
+
+		return true;
+	}
+
+	bool FToolMenuAutomationTestAdapter::Matches(const FMenuSection& InExpectedSection, const FToolMenuSection& InActualSection, const bool bInTestChildren)
+	{
+		bool bAllSucceeded = true;
+
+		if (!InExpectedSection.Name.IsNone()
+			&& !TestInstance.TestEqual("Section Name", InActualSection.Name, InExpectedSection.Name))
+		{
+			bAllSucceeded = false;
+		}
+
+		if (!bInTestChildren)
+		{
+			return true;
+		}
+
+		// If we don't require all entries, we default to failure, and only succeed if we find all entries in the (expected) section
+		if (!Parameters.EntryMatchParameters.bHasAllItems)
+		{
+			bAllSucceeded = false;
+		}
+
+		// If we require an exact match, we require an exact count
+		if ((Parameters.EntryMatchParameters.bHasAllItems && !TestInstance.TestEqual(
+				FString::Printf(TEXT("Entry Count (for section \"%s\""), *InActualSection.Name.ToString()),
+				InActualSection.Blocks.Num(),
+				InExpectedSection.Entries.Num()))
+			// ...otherwise we should have AT LEAST the expected item count
+			|| (!Parameters.EntryMatchParameters.bHasAllItems && !TestInstance.TestGreaterEqual(
+				FString::Printf(TEXT("Entry Count (for section \"%s\""), *InActualSection.Name.ToString()),
+				InActualSection.Blocks.Num(),
+				InExpectedSection.Entries.Num())))
+		{
+			// ...we default to failure but continue checking for matches before returning
+			bAllSucceeded = false;
+		}
+
+		// If we don't require all entries to be specified in the expected section, keep track of the amount we expect to find and only succeed if we find them all
+		const int32 ExpectedEntryCount = InExpectedSection.Entries.Num();
+		int32 FoundEntryCount = 0;
+
+		// Track the last found index, to validate order if required
+		int32 LastFoundEntryIdx = INDEX_NONE;
+		bool bFoundEntriesInOrder = true;
+
+		int32 ActualEntryIdx = 0;
+
+		for (int32 ExpectedEntryIdx = 0; ExpectedEntryIdx < InExpectedSection.Entries.Num(); ++ExpectedEntryIdx)
+		{
+			const TSharedRef<FMenuEntryBase>& ExpectedEntry = InExpectedSection.Entries[ExpectedEntryIdx];
+
+			for (; ActualEntryIdx < InActualSection.Blocks.Num(); ++ActualEntryIdx)
+			{
+				if (!TestValidIndex(FString::Printf(TEXT("Entry (for section \"%s\")"), *InActualSection.Name.ToString()), ActualEntryIdx, InActualSection.Blocks))
+				{
+					bAllSucceeded = false;
+					continue;
+				}
+
+				const FToolMenuEntry& ActualEntry = InActualSection.Blocks[ActualEntryIdx];
+
+				if (ExpectedEntry->Name == ActualEntry.Name)
+				{
+					if (ExpectedEntry->GetTypeName() == FMenuEntry::TypeName)
+					{
+						if (!Matches(StaticCastSharedRef<FMenuEntry>(ExpectedEntry).Get(), ActualEntry))
+						{
+							bAllSucceeded = false;
+						}
+					}
+
+					if (!Matches(&ExpectedEntry.Get(), ActualEntry))
+					{
+						bAllSucceeded = false;
+					}
+					else
+					{
+						// We matched, increment found
+						++FoundEntryCount;
+
+						if (Parameters.EntryMatchParameters.bHasSameOrder)
+						{
+							if (ActualEntryIdx < LastFoundEntryIdx)
+							{
+								bFoundEntriesInOrder = false;
+							}
+
+							LastFoundEntryIdx = ActualEntryIdx;
+						}
+
+						break;
+					}
+				}
+			}
+
+			// If the order doesn't matter, we want to search the entire actual menu, not just subsequent items, so reset the index
+			if (!Parameters.EntryMatchParameters.bHasSameOrder)
+			{
+				ActualEntryIdx = 0;
+			}
+		}
+
+		if (!Parameters.EntryMatchParameters.bHasAllItems)
+		{
+			// All expected entries were found - success!
+			bAllSucceeded = (Parameters.EntryMatchParameters.bHasSameOrder && bFoundEntriesInOrder)
+						&& (FoundEntryCount == ExpectedEntryCount);
+		}
+
+ 		return bAllSucceeded;
+	}
+
+	bool FToolMenuAutomationTestAdapter::Matches(const FMenu& InExpectedMenu, const UToolMenu& InActualMenu)
+	{
+		bool bAllSucceeded = true;
+
+		if (!InExpectedMenu.Name.IsNone()
+			&& !TestInstance.TestEqual("ToolMenu Name", InActualMenu.MenuName, InExpectedMenu.Name))
+		{
+			bAllSucceeded = false;
+		}
+
+		// If we don't require all sections, we default to failure, and only succeed if we find all sections in the (expected) menu
+		if (!Parameters.SectionMatchParameters.bHasAllItems)
+		{
+			bAllSucceeded = false;
+		}
+
+		// If we require an exact match, we require an exact count
+		if ((Parameters.SectionMatchParameters.bHasAllItems && !TestInstance.TestEqual(
+				FString::Printf(TEXT("Section Count (for menu \"%s\""), *InActualMenu.MenuName.ToString()),
+				InActualMenu.Sections.Num(),
+				InExpectedMenu.Sections.Num()))
+			// ...otherwise we should have AT LEAST the expected item count
+			|| (!Parameters.SectionMatchParameters.bHasAllItems && !TestInstance.TestGreaterEqual(
+				FString::Printf(TEXT("Section Count (for menu \"%s\""), *InActualMenu.MenuName.ToString()),
+				InActualMenu.Sections.Num(),
+				InExpectedMenu.Sections.Num())))
+		{
+			// ...we default to failure but continue checking for matches before returning
+			bAllSucceeded = false;
+		}
+
+		// If we don't require all sections to be specified in the expected menu, keep track of the amount we expect to find and only succeed if we find them all
+		const int32 ExpectedSectionCount = InExpectedMenu.Sections.Num();
+		int32 FoundSectionCount = 0;
+
+		// Track the last found index, to validate order if required
+		int32 LastFoundSectionIdx = INDEX_NONE;
+		bool bFoundEntriesInOrder = true;
+
+		int32 ActualSectionIdx = 0;
+
+		for (int32 ExpectedSectionIdx = 0; ExpectedSectionIdx < InExpectedMenu.Sections.Num(); ++ExpectedSectionIdx)
+		{
+			const TSharedRef<FMenuSection>& ExpectedSection = InExpectedMenu.Sections[ExpectedSectionIdx];
+
+			for (; ActualSectionIdx < InActualMenu.Sections.Num(); ++ActualSectionIdx)
+			{
+				if (!TestValidIndex(FString::Printf(TEXT("Section (for menu \"%s\")"), *InActualMenu.MenuName.ToString()), ActualSectionIdx, InActualMenu.Sections))
+				{
+					bAllSucceeded = false;
+					continue;
+				}
+
+				const FToolMenuSection& ActualSection = InActualMenu.Sections[ActualSectionIdx];
+
+				if (ExpectedSection->Name == ActualSection.Name)
+				{
+					if (!Matches(*ExpectedSection, ActualSection, Parameters.SectionMatchParameters.bMatchChildrenIfAnyExpected && ExpectedSection->Entries.Num() > 0))
+					{
+						bAllSucceeded = false;
+					}
+					else
+					{
+						// We matched, increment found
+						++FoundSectionCount;
+
+						if (Parameters.SectionMatchParameters.bHasSameOrder)
+						{
+							if (ActualSectionIdx < LastFoundSectionIdx)
+							{
+								bFoundEntriesInOrder = false;
+							}
+
+							LastFoundSectionIdx = ActualSectionIdx;
+						}
+
+						break;
+					}
+				}
+			}
+
+			// If the order doesn't matter, we want to search the entire actual menu, not just subsequent items, so reset the index
+			if (!Parameters.SectionMatchParameters.bHasSameOrder)
+			{
+				ActualSectionIdx = 0;
+			}
+		}
+
+		if (!Parameters.SectionMatchParameters.bHasAllItems)
+		{
+			// All expected sections were found - success!
+			bAllSucceeded = (Parameters.SectionMatchParameters.bHasSameOrder && bFoundEntriesInOrder)
+						&& (FoundSectionCount == ExpectedSectionCount);
+		}
+
+		return bAllSucceeded;
+	}
+}
diff --git a/Engine/Source/Developer/ToolMenus/Private/Tests/ToolMenusTestUtilities.spec.cpp b/Engine/Source/Developer/ToolMenus/Private/Tests/ToolMenusTestUtilities.spec.cpp
new file mode 100644
--- /dev/null
+++ b/Engine/Source/Developer/ToolMenus/Private/Tests/ToolMenusTestUtilities.spec.cpp
@@ -0,0 +1,232 @@
+// Copyright Epic Games, Inc. All Rights Reserved.
+
+#if WITH_TESTS
+
+#include "Misc/AutomationTest.h"
+#include "Tests/ToolMenusTestUtilities.h"
+
+namespace UE::ToolMenus::Private
+{
+	BEGIN_DEFINE_SPEC(FToolMenusTestUtiltiesSpec, "System.ToolMenus.TestUtilities", EAutomationTestFlags::EditorContext | EAutomationTestFlags::EngineFilter)
+	static const FName TestMenuName;
+	void RegisterTestMenu();
+	const UToolMenu* GenerateTestMenu() const;
+	bool TestMatch(const FString& InWhat, const FMenu& InExpectedMenu, const UToolMenu& InActualMenu, const FToolMenuAutomationTestAdapter::FToolMenuAutomationTestAdapterParameters& InParameters = {});
+	bool TestMismatch(const FString& InWhat, const FMenu& InExpectedMenu, const UToolMenu& InActualMenu, const FToolMenuAutomationTestAdapter::FToolMenuAutomationTestAdapterParameters& InParameters = {});
+	END_DEFINE_SPEC(FToolMenusTestUtiltiesSpec)
+
+	const FName FToolMenusTestUtiltiesSpec::TestMenuName = "TestMenu";
+
+	void FToolMenusTestUtiltiesSpec::Define()
+	{
+		RegisterTestMenu();
+
+		Describe("Exact Match", [this]()
+		{
+			It("Should match the exact sections", [this]()
+			{
+				const FMenu ExpectedMenuStructure(
+					"TestMenu",
+					{
+						Section("FirstSection"),
+						Section("MiddleSection"),
+						Section("LastSection")
+					});
+
+				const UToolMenu* ActualMenu = GenerateTestMenu();
+
+				FToolMenuAutomationTestAdapter::FToolMenuAutomationTestAdapterParameters TestAdapterParameters;
+				TestAdapterParameters.SectionMatchParameters.bHasAllItems = true;
+				TestAdapterParameters.EntryMatchParameters.bHasAllItems = true;
+
+				TestMatch(
+					"Actual menu sections to (exactly) match the expected sections",
+					ExpectedMenuStructure,
+					*ActualMenu,
+					TestAdapterParameters);
+			});
+
+			It("Should match the exact entries", [this]()
+			{
+				const FMenu ExpectedMenuStructure(
+					"TestMenu",
+					{
+						Section("FirstSection", {
+							Entry("FirstEntry"),
+							Entry("MiddleEntry"),
+							Separator(),
+							Entry("LastEntry")
+						}),
+						Section("MiddleSection"),
+						Section("LastSection")
+					});
+
+				const UToolMenu* ActualMenu = GenerateTestMenu();
+
+				FToolMenuAutomationTestAdapter::FToolMenuAutomationTestAdapterParameters TestAdapterParameters;
+				TestAdapterParameters.SectionMatchParameters.bHasAllItems = false;
+				TestAdapterParameters.EntryMatchParameters.bHasAllItems = true;
+
+				TestMatch(
+					"Actual section entries to (exactly) match the expected section entries",
+					ExpectedMenuStructure,
+					*ActualMenu,
+					TestAdapterParameters);
+			});
+		});
+
+		Describe("Partial Match", [this]()
+		{
+			It("Should match any existing section", [this]()
+			{
+				const FMenu ExpectedPartialMenuStructure(
+					"TestMenu",
+					{
+						Section("MiddleSection")
+					});
+
+				const UToolMenu* ActualMenu = GenerateTestMenu();
+
+				FToolMenuAutomationTestAdapter::FToolMenuAutomationTestAdapterParameters TestAdapterParameters;
+				TestAdapterParameters.SectionMatchParameters.bHasAllItems = false;
+				TestAdapterParameters.EntryMatchParameters.bHasAllItems = false;
+
+				TestMatch(
+					"Actual menu sections to (partially) match the expected section(s)",
+					ExpectedPartialMenuStructure,
+					*ActualMenu,
+					TestAdapterParameters);
+			});
+
+			It("Shouldn't match out-of-order sections", [this]()
+			{
+				const FMenu WrongOrderMenuStructure(
+					"TestMenu",
+					{
+						Section("FirstSection"),
+						Section("LastSection"),
+						Section("MiddleSection")
+					});
+
+				const UToolMenu* ActualMenu = GenerateTestMenu();
+				FToolMenuAutomationTestAdapter::FToolMenuAutomationTestAdapterParameters TestAdapterParameters;
+				TestAdapterParameters.SectionMatchParameters.bHasAllItems = false;
+				TestAdapterParameters.SectionMatchParameters.bHasSameOrder = true;
+				TestAdapterParameters.EntryMatchParameters.bHasAllItems = false;
+
+				TestMismatch(
+					"Menu sections shouldn't match wrongly ordered sections",
+					WrongOrderMenuStructure,
+					*ActualMenu,
+					TestAdapterParameters);
+			});
+
+			It("Should match any existing entry", [this]()
+			{
+				const FMenu ExpectedPartialMenuStructure(
+					"TestMenu",
+					{
+						Section("FirstSection", {
+							Entry("MiddleEntry"),
+							Entry("LastEntry")
+						})
+					});
+
+				const UToolMenu* ActualMenu = GenerateTestMenu();
+
+				FToolMenuAutomationTestAdapter::FToolMenuAutomationTestAdapterParameters TestAdapterParameters;
+				TestAdapterParameters.SectionMatchParameters.bHasAllItems = false;
+				TestAdapterParameters.EntryMatchParameters.bHasAllItems = false;
+
+				TestMatch(
+					"Actual menu entries to (partially) match the expected entries(s)",
+					ExpectedPartialMenuStructure,
+					*ActualMenu,
+					TestAdapterParameters);
+			});
+
+			It("Shouldn't match out-of-order entries", [this]()
+			{
+				const FMenu ExpectedPartialMenuStructure(
+					"TestMenu",
+					{
+						Section("FirstSection", {
+							Entry("LastEntry"),
+							Entry("MiddleEntry")
+						})
+					});
+
+				const UToolMenu* ActualMenu = GenerateTestMenu();
+
+				FToolMenuAutomationTestAdapter::FToolMenuAutomationTestAdapterParameters TestAdapterParameters;
+				TestAdapterParameters.SectionMatchParameters.bHasAllItems = false;
+				TestAdapterParameters.EntryMatchParameters.bHasAllItems = false;
+				TestAdapterParameters.EntryMatchParameters.bHasSameOrder = true;
+
+				TestMismatch(
+					"Menu entries shouldn't match wrongly ordered entries",
+					ExpectedPartialMenuStructure,
+					*ActualMenu,
+					TestAdapterParameters);
+			});
+		});
+	}
+
+	// Registers the actual menu to test against
+	void FToolMenusTestUtiltiesSpec::RegisterTestMenu()
+	{
+		if (!UToolMenus::Get()->IsMenuRegistered(TestMenuName))
+		{
+			UToolMenu* TestMenu = UToolMenus::Get()->RegisterMenu(TestMenuName);
+
+			auto MakeEntry = [](const FName Name)
+			{
+				return FToolMenuEntry::InitMenuEntry(Name, FText::GetEmpty(), FText::GetEmpty(), FSlateIcon(), FUIAction(), EUserInterfaceActionType::Button);
+			};
+
+			FToolMenuSection& FirstSection = TestMenu->AddSection("FirstSection");
+			{
+				FirstSection.AddEntry(MakeEntry("FirstEntry"));
+				FirstSection.AddEntry(MakeEntry("MiddleEntry"));
+				FirstSection.AddSeparator("FirstSeparator");
+				FirstSection.AddEntry(MakeEntry("LastEntry"));
+			}
+
+			FToolMenuSection& MiddleSection = TestMenu->AddSection("MiddleSection");
+
+			FToolMenuSection& LastSection = TestMenu->AddSection("LastSection");
+		}
+	}
+
+	const UToolMenu* FToolMenusTestUtiltiesSpec::GenerateTestMenu() const
+	{
+		FToolMenuContext ToolMenuContext;
+		return UToolMenus::Get()->GenerateMenu(TestMenuName, ToolMenuContext);
+	}
+
+	bool FToolMenusTestUtiltiesSpec::TestMatch(
+		const FString& InWhat,
+		const FMenu& InExpectedMenu,
+		const UToolMenu& InActualMenu,
+		const FToolMenuAutomationTestAdapter::FToolMenuAutomationTestAdapterParameters& InParameters)
+	{
+		return TestTrue(
+			InWhat,
+			FToolMenuAutomationTestAdapter(*this, InParameters)
+			.Matches(InExpectedMenu, InActualMenu));
+	}
+
+	bool FToolMenusTestUtiltiesSpec::TestMismatch(
+		const FString& InWhat,
+		const FMenu& InExpectedMenu,
+		const UToolMenu& InActualMenu,
+		const FToolMenuAutomationTestAdapter::FToolMenuAutomationTestAdapterParameters& InParameters)
+	{
+		return TestFalse(
+			InWhat,
+			FToolMenuAutomationTestAdapter(*this, InParameters)
+			.Matches(InExpectedMenu, InActualMenu));
+	}
+}
+
+#endif
diff --git a/Engine/Source/Developer/ToolMenus/Public/Tests/ToolMenusTestUtilities.h b/Engine/Source/Developer/ToolMenus/Public/Tests/ToolMenusTestUtilities.h
new file mode 100644
--- /dev/null
+++ b/Engine/Source/Developer/ToolMenus/Public/Tests/ToolMenusTestUtilities.h
@@ -0,0 +1,299 @@
+// Copyright Epic Games, Inc. All Rights Reserved.
+
+#pragma once
+
+#include "Framework/MultiBox/MultiBoxDefs.h"
+#include "Misc/AutomationTest.h"
+#include "ToolMenu.h"
+#include "ToolMenus.h"
+
+namespace UE::ToolMenus
+{
+	// @todo: account for extension points, ie. allow for other arbitrary entries, but enforce order of known entries to test against
+
+	enum class EExpectedOccurrence : uint8
+	{
+		Any,
+		ExactlyZero,
+		ExactlyOne,
+		OneOrMore
+	};
+
+	/**
+	 * @note: that these parameters apply to the item type, not the children of that item (with some exceptions).
+	 * When used for sections, it applies to Sections in general and not entries within that section.
+	 */
+	struct TOOLMENUS_API FMenuMatchParameters
+	{
+		/** All expected entries/sections must exist in the actual menu. */
+		bool bHasAllItems = false;
+
+		/** The expected entries/sections must be in the same order they are in the actual menu. */
+		bool bHasSameOrder = true;
+
+		/** Allow extra entries in the actual menu, skipping over them if they're found in actual but not expected. Relevant when bHasAllItems is true. */
+		bool bAllowExtraItems = true;
+
+		/** Attempts to match children if any are expected, otherwise it skips them. */
+		bool bMatchChildrenIfAnyExpected = true;
+	};
+
+	struct TOOLMENUS_API FMenuEntryBase
+	{
+		explicit FMenuEntryBase(const FName InName, const EExpectedOccurrence InOccurrence = EExpectedOccurrence::ExactlyOne)
+			: Name(InName)
+			, Occurrence(InOccurrence)
+		{
+		}
+
+		virtual ~FMenuEntryBase() = default;
+
+		FName Name;
+		EExpectedOccurrence Occurrence = EExpectedOccurrence::ExactlyOne;
+
+		virtual FName GetTypeName() = 0;
+	};
+
+	/** Represents a Named ToolMenuEntry, optionally checking the type. */
+	struct TOOLMENUS_API FMenuEntry
+		: FMenuEntryBase
+		, public TSharedFromThis<FMenuEntry>
+	{
+		explicit FMenuEntry(const FName InName)
+			: FMenuEntryBase(InName)
+		{
+		}
+
+		virtual ~FMenuEntry() override = default;
+
+		TOptional<FString> Label;
+		TOptional<EMultiBlockType> Type = {};
+
+		TSharedRef<FMenuEntry> WithLabel(const FString& InLabel)
+		{
+			Label = InLabel;
+			return AsShared();
+		}
+
+		bool operator==(const FToolMenuEntry& InToolMenuEntry) const
+		{
+			return !Name.IsNone() && Name == InToolMenuEntry.Name
+				&& (!Label.IsSet() || InToolMenuEntry.Label.Get().ToString().MatchesWildcard(Label.GetValue()))
+				&& (!Type.IsSet() || Type.GetValue() == InToolMenuEntry.Type);
+		}
+
+		static FName TypeName;
+
+		virtual FName GetTypeName() override
+		{
+			return TypeName;
+		}
+	};
+
+	/** Represents the presence of any MenuEntry, used when the test is agnostic to what's placed/inserted into an extension point. */
+	struct TOOLMENUS_API FMenuWildcardEntry : FMenuEntryBase
+	{
+		FMenuWildcardEntry()
+			: FMenuEntryBase(NAME_None, EExpectedOccurrence::Any)
+		{
+		}
+
+		virtual ~FMenuWildcardEntry() override = default;
+
+		bool operator==(TConstArrayView<FToolMenuEntry> InToolMenuEntries) const
+		{
+			// @todo
+			return true;
+		}
+
+		static FName TypeName;
+
+		virtual FName GetTypeName() override
+		{
+			return TypeName;
+		}
+	};
+
+	/** Represents a Named ToolMenuSection with zero or more (sorted) Menu Entries. */
+	struct TOOLMENUS_API FMenuSection
+		: public TSharedFromThis<FMenuSection>
+	{
+		explicit FMenuSection(const FName InName, const TArray<TSharedRef<FMenuEntryBase>>& InEntries = {})
+			: Name(InName)
+			, Entries(InEntries)
+		{
+		}
+
+		FName Name;
+		TOptional<FString> Label;
+		TArray<TSharedRef<FMenuEntryBase>> Entries;
+
+		TSharedRef<FMenuSection> WithLabel(const FString& InLabel)
+		{
+			Label = InLabel;
+			return AsShared();
+		}
+
+		TSharedRef<FMenuSection> WithEntries(const TArray<TSharedRef<FMenuEntryBase>>& InEntries)
+		{
+			Entries = InEntries;
+			return AsShared();
+		}
+
+		/** Only compare top-level item, ignore children (entries). */
+		bool operator==(const FToolMenuSection& InToolMenuSection) const
+		{
+			return (!Name.IsNone() || Name == InToolMenuSection.Name)
+				&& (!Label.IsSet() || InToolMenuSection.Label.Get().ToString().MatchesWildcard(Label.GetValue()));
+		}
+	};
+
+	struct TOOLMENUS_API FMenu
+		: FMenuEntryBase
+		, public TSharedFromThis<FMenu>
+	{
+		explicit FMenu(const FName InName, const TArray<TSharedRef<FMenuSection>>& InSections = {})
+			: FMenuEntryBase(InName)
+			, Sections(InSections)
+		{
+		}
+
+		virtual ~FMenu() override = default;
+
+		TArray<TSharedRef<FMenuSection>> Sections;
+
+        TSharedRef<FMenu> WithSections(const TArray<TSharedRef<FMenuSection>>& InSections)
+        {
+			Sections = InSections;
+        	return AsShared();
+        }
+
+		/** Only compare top-level item, ignore children (sections). */
+		bool operator==(const UToolMenu& InToolMenu) const
+		{
+        	return Name == InToolMenu.MenuName;
+		}
+
+		static FName TypeName;
+
+		virtual FName GetTypeName() override
+		{
+			return TypeName;
+		}
+	};
+
+	/**
+	 * The factory functions below reduce verbosity when constructing Expected menu structures, and can be used like so:
+	 * @example
+	 * using namespace UE::ToolMenus;
+     * const FMenu ExpectedMenuStructure(
+     *		"SomeMenu",
+     *		{
+     *			Section("SomeSection",
+     *			{
+     *				Entry("SomeEntry"),
+     *				Any(),
+     *				Entry("AnotherEntry")
+	 *			})
+     * 		});
+	 */
+
+	inline TSharedRef<FMenuEntry> Entry()
+	{
+		return MakeShared<FMenuEntry>(NAME_None);
+	}
+
+	/** Creates a MenuEntry with a label to match, with wildcards. This will check if the actual entry being tested against matches the given label. */
+	inline TSharedRef<FMenuEntry> Entry(const FName InName)
+	{
+		return MakeShared<FMenuEntry>(InName);
+	}
+
+	inline TSharedRef<FMenuEntry> Separator()
+	{
+		TSharedRef<FMenuEntry> Result = MakeShared<FMenuEntry>(NAME_All);
+		Result->Type = EMultiBlockType::Separator;
+		return Result;
+	}
+
+	inline TSharedRef<FMenuEntry> Any(const EExpectedOccurrence InOccurence = EExpectedOccurrence::Any)
+	{
+		TSharedRef<FMenuEntry> Result = MakeShared<FMenuEntry>(NAME_None);
+		Result->Occurrence = InOccurence;
+		return Result;
+	}
+
+	inline TSharedRef<FMenuSection> Section()
+	{
+		return MakeShared<FMenuSection>(NAME_None);
+	}
+
+	inline TSharedRef<FMenuSection> Section(const FName InName, const TArray<TSharedRef<FMenuEntryBase>>& InEntries = {})
+	{
+		return MakeShared<FMenuSection>(InName, InEntries);
+	}
+
+	inline TSharedRef<FMenu> SubMenu()
+	{
+		return MakeShared<FMenu>(NAME_None);
+	}
+
+	inline TSharedRef<FMenu> SubMenu(const FName InName, const TArray<TSharedRef<FMenuSection>>& InSections = {})
+	{
+		return MakeShared<FMenu>(InName, InSections);
+	}
+
+	/** @note: Tests will continue so long as they can - we don't return for the first failed case. */
+	class FToolMenuAutomationTestAdapter
+	{
+	public:
+		struct FToolMenuAutomationTestAdapterParameters
+		{
+			FToolMenuAutomationTestAdapterParameters()
+				: SectionMatchParameters({ true, true, true, true })
+				, EntryMatchParameters({ true, true, true, true })
+			{
+
+			}
+
+			/** Note that regardless of match parameters, use of Any() will explicitly allow one or more unknown entries in the specified location. */
+			explicit FToolMenuAutomationTestAdapterParameters(
+				const FMenuMatchParameters& InSectionMatchParameters,
+				const FMenuMatchParameters& InEntryMatchParameters)
+				: SectionMatchParameters(InSectionMatchParameters)
+				, EntryMatchParameters(InEntryMatchParameters)
+			{
+			}
+
+			FMenuMatchParameters SectionMatchParameters;
+			FMenuMatchParameters EntryMatchParameters;
+		};
+
+		explicit FToolMenuAutomationTestAdapter(FAutomationTestBase& InTestInstance, const FToolMenuAutomationTestAdapterParameters& InParameters = {})
+			: TestInstance(InTestInstance)
+			, Parameters(InParameters)
+		{
+		}
+
+		TOOLMENUS_API bool Matches(const FMenuEntryBase* InExpectedEntry, const FToolMenuEntry& InActualEntry);
+		TOOLMENUS_API bool Matches(const FMenuEntry& InExpectedEntry, const FToolMenuEntry& InActualEntry);
+		TOOLMENUS_API bool Matches(const FMenuSection& InExpectedSection, const FToolMenuSection& InActualSection, bool bInTestChildren = true);
+		TOOLMENUS_API bool Matches(const FMenu& InExpectedMenu, const UToolMenu& InActualMenu);
+
+	private:
+		template <typename ElementType>
+		bool TestValidIndex(const FString& What, const int32 Index, const TArray<ElementType>& Array)
+		{
+			if (!Array.IsValidIndex(Index))
+			{
+				TestInstance.AddError(FString::Printf(TEXT("Expected %s at Index %d, but the array only contains %d elements."), *What, Index, Array.Num()), 1);
+			}
+
+			return true;
+		}
+
+	private:
+		FAutomationTestBase& TestInstance;
+		FToolMenuAutomationTestAdapterParameters Parameters;
+	};
+}
