UE: Add multicast delegate to report when string commands are executed. Useful for analytics as it includes the name of the menu command

diff --git a/Engine/Source/Developer/ToolMenus/Private/ToolMenus.cpp b/Engine/Source/Developer/ToolMenus/Private/ToolMenus.cpp
--- a/Engine/Source/Developer/ToolMenus/Private/ToolMenus.cpp
+++ b/Engine/Source/Developer/ToolMenus/Private/ToolMenus.cpp
@@ -68,11 +68,11 @@ FName FToolMenuStringCommand::GetTypeName() const
 	return NAME_None;
 }
 
-FExecuteAction FToolMenuStringCommand::ToExecuteAction(const FToolMenuContext& Context) const
+FExecuteAction FToolMenuStringCommand::ToExecuteAction(FName MenuName, const FToolMenuContext& Context) const
 {
 	if (IsBound())
 	{
-		return FExecuteAction::CreateStatic(&UToolMenus::ExecuteStringCommand, *this, Context);
+		return FExecuteAction::CreateStatic(&UToolMenus::ExecuteStringCommand, *this, MenuName, Context);
 	}
 
 	return FExecuteAction();
@@ -2241,7 +2241,7 @@ FUIAction UToolMenus::ConvertUIAction(const FToolMenuEntry& Block, const FToolMe
 
 	if (!UIAction.ExecuteAction.IsBound() && Block.StringExecuteAction.IsBound())
 	{
-		UIAction.ExecuteAction = Block.StringExecuteAction.ToExecuteAction(Context);
+		UIAction.ExecuteAction = Block.StringExecuteAction.ToExecuteAction(Block.Name, Context);
 	}
 
 	return UIAction;
@@ -2412,7 +2412,7 @@ FUIAction UToolMenus::ConvertScriptObjectToUIAction(UToolMenuEntryScript* Script
 	return UIAction;
 }
 
-void UToolMenus::ExecuteStringCommand(const FToolMenuStringCommand StringCommand, const FToolMenuContext Context)
+void UToolMenus::ExecuteStringCommand(const FToolMenuStringCommand StringCommand, FName MenuName, FToolMenuContext Context)
 {
 	if (StringCommand.IsBound())
 	{
@@ -2429,6 +2429,8 @@ void UToolMenus::ExecuteStringCommand(const FToolMenuStringCommand StringCommand
 		{
 			UE_LOG(LogToolMenus, Warning, TEXT("Unknown string command handler type: '%s'"), *TypeName.ToString());
 		}
+
+		ToolMenus->OnStringCommandExecuted.Broadcast(MenuName, TypeName);
 	}
 }
 
diff --git a/Engine/Source/Developer/ToolMenus/Public/ToolMenuMisc.h b/Engine/Source/Developer/ToolMenus/Public/ToolMenuMisc.h
--- a/Engine/Source/Developer/ToolMenus/Public/ToolMenuMisc.h
+++ b/Engine/Source/Developer/ToolMenus/Public/ToolMenuMisc.h
@@ -49,7 +49,7 @@ private:
 
 	bool IsBound() const { return String.Len() > 0; }
 
-	FExecuteAction ToExecuteAction(const FToolMenuContext& Context) const;
+	FExecuteAction ToExecuteAction(FName MenuName, const FToolMenuContext& Context) const;
 
 	FName GetTypeName() const;
 };
diff --git a/Engine/Source/Developer/ToolMenus/Public/ToolMenus.h b/Engine/Source/Developer/ToolMenus/Public/ToolMenus.h
--- a/Engine/Source/Developer/ToolMenus/Public/ToolMenus.h
+++ b/Engine/Source/Developer/ToolMenus/Public/ToolMenus.h
@@ -257,6 +257,10 @@ public:
 
 	/** Removes a string based command handler. */
 	void UnregisterStringCommandHandler(const FName InName);
+
+	DECLARE_MULTICAST_DELEGATE_TwoParams(FOnStringCommandExecuted, FName MenuName, FName TypeName);
+	/** Called when we execute a string command. Useful for analytics */
+	FOnStringCommandExecuted OnStringCommandExecuted;
 	
 	/** Sets delegate to setup timer for deferred one off ticks */
 	void AssignSetTimerForNextTickDelegate(const FSimpleDelegate& InDelegate);
@@ -416,7 +420,7 @@ private:
 	static FUIAction ConvertUIAction(const FToolDynamicUIAction& Actions, const FToolMenuContext& Context);
 	static FUIAction ConvertScriptObjectToUIAction(UToolMenuEntryScript* ScriptObject, const FToolMenuContext& Context);
 
-	static void ExecuteStringCommand(const FToolMenuStringCommand StringCommand, const FToolMenuContext Context);
+	static void ExecuteStringCommand(const FToolMenuStringCommand StringCommand, FName MenuName, const FToolMenuContext Context);
 
 	void PopulateSubMenu(FMenuBuilder& Builder, TWeakObjectPtr<UToolMenu> InParent, const FToolMenuEntry InEntry, const FName InBlockName);
 	void PopulateSubMenuWithoutName(FMenuBuilder& MenuBuilder, TWeakObjectPtr<UToolMenu> InParent, const FToolMenuEntry InEntry);
