[ToolMenus] TestUtilities fixes

diff --git a/Engine/Source/Developer/ToolMenus/Private/Tests/ToolMenusTestUtilities.cpp b/Engine/Source/Developer/ToolMenus/Private/Tests/ToolMenusTestUtilities.cpp
--- a/Engine/Source/Developer/ToolMenus/Private/Tests/ToolMenusTestUtilities.cpp
+++ b/Engine/Source/Developer/ToolMenus/Private/Tests/ToolMenusTestUtilities.cpp
@@ -57,21 +57,40 @@ namespace UE::ToolMenus
 		}
 
 		// If we don't require all entries, we default to failure, and only succeed if we find all entries in the (expected) section
-		if (!Parameters.EntryMatchParameters.bHasAllItems)
+		if (!Parameters.EntryMatchParameters.bActualHasAllExpectedItems)
 		{
 			bAllSucceeded = false;
 		}
 
 		// If we require an exact match, we require an exact count
-		if ((Parameters.EntryMatchParameters.bHasAllItems && !TestInstance.TestEqual(
-				FString::Printf(TEXT("Entry Count (for section \"%s\""), *InActualSection.Name.ToString()),
+		const bool bRequireMatchingEntryCount = Parameters.EntryMatchParameters.bActualHasAllExpectedItems && Parameters.EntryMatchParameters.bExpectedHasAllActualItems;
+		if (bRequireMatchingEntryCount
+			&& !TestInstance.TestEqual(
+				FString::Printf(TEXT("Entry Count (for section \"%s\")"), *InActualSection.Name.ToString()),
 				InActualSection.Blocks.Num(),
 				InExpectedSection.Entries.Num()))
-			// ...otherwise we should have AT LEAST the expected item count
-			|| (!Parameters.EntryMatchParameters.bHasAllItems && !TestInstance.TestGreaterEqual(
-				FString::Printf(TEXT("Entry Count (for section \"%s\""), *InActualSection.Name.ToString()),
+		{
+			// ...we default to failure but continue checking for matches before returning
+			bAllSucceeded = false;
+		}
+
+		// Actual should have at least the expected item count
+		if (Parameters.EntryMatchParameters.bActualHasAllExpectedItems
+			&& !TestInstance.TestGreaterEqual(
+				FString::Printf(TEXT("Entry Count (for section \"%s\")"), *InActualSection.Name.ToString()),
 				InActualSection.Blocks.Num(),
-				InExpectedSection.Entries.Num())))
+				InExpectedSection.Entries.Num()))
+		{
+			// ...we default to failure but continue checking for matches before returning
+			bAllSucceeded = false;
+		}
+
+		// Expected should have at least the actual item count
+		if (Parameters.EntryMatchParameters.bExpectedHasAllActualItems
+			&& !TestInstance.TestGreaterEqual(
+				FString::Printf(TEXT("Entry Count (for section \"%s\")"), *InActualSection.Name.ToString()),
+				InExpectedSection.Entries.Num(),
+				InActualSection.Blocks.Num()))
 		{
 			// ...we default to failure but continue checking for matches before returning
 			bAllSucceeded = false;
@@ -85,13 +104,11 @@ namespace UE::ToolMenus
 		int32 LastFoundEntryIdx = INDEX_NONE;
 		bool bFoundEntriesInOrder = true;
 
-		int32 ActualEntryIdx = 0;
-
 		for (int32 ExpectedEntryIdx = 0; ExpectedEntryIdx < InExpectedSection.Entries.Num(); ++ExpectedEntryIdx)
 		{
 			const TSharedRef<FMenuEntryBase>& ExpectedEntry = InExpectedSection.Entries[ExpectedEntryIdx];
 
-			for (; ActualEntryIdx < InActualSection.Blocks.Num(); ++ActualEntryIdx)
+			for (int32 ActualEntryIdx = 0; ActualEntryIdx < InActualSection.Blocks.Num(); ++ActualEntryIdx)
 			{
 				if (!TestValidIndex(FString::Printf(TEXT("Entry (for section \"%s\")"), *InActualSection.Name.ToString()), ActualEntryIdx, InActualSection.Blocks))
 				{
@@ -101,9 +118,12 @@ namespace UE::ToolMenus
 
 				const FToolMenuEntry& ActualEntry = InActualSection.Blocks[ActualEntryIdx];
 
-				if (ExpectedEntry->Name == ActualEntry.Name)
+				// "None" matches anything, so we're a bit looser about matching requirements
+				const bool bMatchAny = ExpectedEntry->Name.IsNone();
+
+				if (bMatchAny || ExpectedEntry->Name == ActualEntry.Name)
 				{
-					if (ExpectedEntry->GetTypeName() == FMenuEntry::TypeName)
+					if (!bMatchAny && ExpectedEntry->GetTypeName() == FMenuEntry::TypeName)
 					{
 						if (!Matches(StaticCastSharedRef<FMenuEntry>(ExpectedEntry).Get(), ActualEntry))
 						{
@@ -111,7 +131,7 @@ namespace UE::ToolMenus
 						}
 					}
 
-					if (!Matches(&ExpectedEntry.Get(), ActualEntry))
+					if (!bMatchAny && !Matches(&ExpectedEntry.Get(), ActualEntry))
 					{
 						bAllSucceeded = false;
 					}
@@ -120,7 +140,7 @@ namespace UE::ToolMenus
 						// We matched, increment found
 						++FoundEntryCount;
 
-						if (Parameters.EntryMatchParameters.bHasSameOrder)
+						if (Parameters.EntryMatchParameters.bActualHasExpectedOrder)
 						{
 							if (ActualEntryIdx < LastFoundEntryIdx)
 							{
@@ -134,20 +154,18 @@ namespace UE::ToolMenus
 					}
 				}
 			}
-
-			// If the order doesn't matter, we want to search the entire actual menu, not just subsequent items, so reset the index
-			if (!Parameters.EntryMatchParameters.bHasSameOrder)
-			{
-				ActualEntryIdx = 0;
-			}
 		}
 
-		if (!Parameters.EntryMatchParameters.bHasAllItems)
+		if (!bAllSucceeded)
 		{
-			// All expected entries were found - success!
-			bAllSucceeded = (Parameters.EntryMatchParameters.bHasSameOrder && bFoundEntriesInOrder)
-						&& (FoundEntryCount == ExpectedEntryCount);
+			const bool bFoundExpectedSections = (!Parameters.EntryMatchParameters.bActualHasAllExpectedItems && FoundEntryCount > 0)
+											|| FoundEntryCount >= ExpectedEntryCount;
+
+			bAllSucceeded = bFoundExpectedSections;
 		}
+		
+		// Finally, check if we found all entries in order if required
+		bAllSucceeded = bAllSucceeded && (!Parameters.EntryMatchParameters.bActualHasExpectedOrder || bFoundEntriesInOrder);
 
  		return bAllSucceeded;
 	}
@@ -163,21 +181,40 @@ namespace UE::ToolMenus
 		}
 
 		// If we don't require all sections, we default to failure, and only succeed if we find all sections in the (expected) menu
-		if (!Parameters.SectionMatchParameters.bHasAllItems)
+		if (!Parameters.SectionMatchParameters.bActualHasAllExpectedItems)
 		{
 			bAllSucceeded = false;
 		}
 
 		// If we require an exact match, we require an exact count
-		if ((Parameters.SectionMatchParameters.bHasAllItems && !TestInstance.TestEqual(
-				FString::Printf(TEXT("Section Count (for menu \"%s\""), *InActualMenu.MenuName.ToString()),
+		const bool bRequireMatchingSectionCount = Parameters.SectionMatchParameters.bActualHasAllExpectedItems && Parameters.SectionMatchParameters.bExpectedHasAllActualItems;
+		if (bRequireMatchingSectionCount
+			&& !TestInstance.TestEqual(
+				FString::Printf(TEXT("Section Count (for menu \"%s\")"), *InActualMenu.MenuName.ToString()),
 				InActualMenu.Sections.Num(),
 				InExpectedMenu.Sections.Num()))
-			// ...otherwise we should have AT LEAST the expected item count
-			|| (!Parameters.SectionMatchParameters.bHasAllItems && !TestInstance.TestGreaterEqual(
-				FString::Printf(TEXT("Section Count (for menu \"%s\""), *InActualMenu.MenuName.ToString()),
+		{
+			// ...we default to failure but continue checking for matches before returning
+			bAllSucceeded = false;
+		}
+
+		// Actual should have at least the expected item count
+		if (Parameters.SectionMatchParameters.bActualHasAllExpectedItems
+			&& !TestInstance.TestGreaterEqual(
+				FString::Printf(TEXT("Section Count (for menu \"%s\")"), *InActualMenu.MenuName.ToString()),
 				InActualMenu.Sections.Num(),
-				InExpectedMenu.Sections.Num())))
+				InExpectedMenu.Sections.Num()))
+		{
+			// ...we default to failure but continue checking for matches before returning
+			bAllSucceeded = false;
+		}
+
+		// Expected should have at least the actual item count
+		if (Parameters.SectionMatchParameters.bExpectedHasAllActualItems
+			&& !TestInstance.TestGreaterEqual(
+				FString::Printf(TEXT("Section Count (for menu \"%s\")"), *InExpectedMenu.Name.ToString()),
+				InExpectedMenu.Sections.Num(),
+				InActualMenu.Sections.Num()))
 		{
 			// ...we default to failure but continue checking for matches before returning
 			bAllSucceeded = false;
@@ -191,13 +228,11 @@ namespace UE::ToolMenus
 		int32 LastFoundSectionIdx = INDEX_NONE;
 		bool bFoundEntriesInOrder = true;
 
-		int32 ActualSectionIdx = 0;
-
 		for (int32 ExpectedSectionIdx = 0; ExpectedSectionIdx < InExpectedMenu.Sections.Num(); ++ExpectedSectionIdx)
 		{
 			const TSharedRef<FMenuSection>& ExpectedSection = InExpectedMenu.Sections[ExpectedSectionIdx];
 
-			for (; ActualSectionIdx < InActualMenu.Sections.Num(); ++ActualSectionIdx)
+			for (int32 ActualSectionIdx = 0; ActualSectionIdx < InActualMenu.Sections.Num(); ++ActualSectionIdx)
 			{
 				if (!TestValidIndex(FString::Printf(TEXT("Section (for menu \"%s\")"), *InActualMenu.MenuName.ToString()), ActualSectionIdx, InActualMenu.Sections))
 				{
@@ -207,9 +242,13 @@ namespace UE::ToolMenus
 
 				const FToolMenuSection& ActualSection = InActualMenu.Sections[ActualSectionIdx];
 
-				if (ExpectedSection->Name == ActualSection.Name)
+				// "None" matches anything, so we're a bit looser about matching requirements
+				const bool bMatchAny = ExpectedSection->Name.IsNone();
+
+				// "None" matches anything
+				if (bMatchAny || ExpectedSection->Name == ActualSection.Name)
 				{
-					if (!Matches(*ExpectedSection, ActualSection, Parameters.SectionMatchParameters.bMatchChildrenIfAnyExpected && ExpectedSection->Entries.Num() > 0))
+					if (!bMatchAny && !Matches(*ExpectedSection, ActualSection, Parameters.SectionMatchParameters.bMatchChildrenIfAnyExpected && ExpectedSection->Entries.Num() > 0))
 					{
 						bAllSucceeded = false;
 					}
@@ -218,7 +257,7 @@ namespace UE::ToolMenus
 						// We matched, increment found
 						++FoundSectionCount;
 
-						if (Parameters.SectionMatchParameters.bHasSameOrder)
+						if (Parameters.SectionMatchParameters.bActualHasExpectedOrder)
 						{
 							if (ActualSectionIdx < LastFoundSectionIdx)
 							{
@@ -232,21 +271,19 @@ namespace UE::ToolMenus
 					}
 				}
 			}
-
-			// If the order doesn't matter, we want to search the entire actual menu, not just subsequent items, so reset the index
-			if (!Parameters.SectionMatchParameters.bHasSameOrder)
-			{
-				ActualSectionIdx = 0;
-			}
 		}
 
-		if (!Parameters.SectionMatchParameters.bHasAllItems)
+		if (!bAllSucceeded)
 		{
-			// All expected sections were found - success!
-			bAllSucceeded = (Parameters.SectionMatchParameters.bHasSameOrder && bFoundEntriesInOrder)
-						&& (FoundSectionCount == ExpectedSectionCount);
+			const bool bFoundExpectedSections = (!Parameters.SectionMatchParameters.bActualHasAllExpectedItems && FoundSectionCount > 0)
+											|| FoundSectionCount >= ExpectedSectionCount;
+
+			bAllSucceeded = bFoundExpectedSections;
 		}
 
+		// Finally, check if we found all entries in order if required
+		bAllSucceeded = bAllSucceeded && (!Parameters.SectionMatchParameters.bActualHasExpectedOrder || bFoundEntriesInOrder);
+
 		return bAllSucceeded;
 	}
 }
diff --git a/Engine/Source/Developer/ToolMenus/Private/Tests/ToolMenusTestUtilities.spec.cpp b/Engine/Source/Developer/ToolMenus/Private/Tests/ToolMenusTestUtilities.spec.cpp
--- a/Engine/Source/Developer/ToolMenus/Private/Tests/ToolMenusTestUtilities.spec.cpp
+++ b/Engine/Source/Developer/ToolMenus/Private/Tests/ToolMenusTestUtilities.spec.cpp
@@ -36,8 +36,8 @@ namespace UE::ToolMenus::Private
 				const UToolMenu* ActualMenu = GenerateTestMenu();
 
 				FToolMenuAutomationTestAdapter::FToolMenuAutomationTestAdapterParameters TestAdapterParameters;
-				TestAdapterParameters.SectionMatchParameters.bHasAllItems = true;
-				TestAdapterParameters.EntryMatchParameters.bHasAllItems = true;
+				TestAdapterParameters.SectionMatchParameters.bActualHasAllExpectedItems = true;
+				TestAdapterParameters.EntryMatchParameters.bActualHasAllExpectedItems = true;
 
 				TestMatch(
 					"Actual menu sections to (exactly) match the expected sections",
@@ -64,8 +64,8 @@ namespace UE::ToolMenus::Private
 				const UToolMenu* ActualMenu = GenerateTestMenu();
 
 				FToolMenuAutomationTestAdapter::FToolMenuAutomationTestAdapterParameters TestAdapterParameters;
-				TestAdapterParameters.SectionMatchParameters.bHasAllItems = false;
-				TestAdapterParameters.EntryMatchParameters.bHasAllItems = true;
+				TestAdapterParameters.SectionMatchParameters.bActualHasAllExpectedItems = false;
+				TestAdapterParameters.EntryMatchParameters.bActualHasAllExpectedItems = true;
 
 				TestMatch(
 					"Actual section entries to (exactly) match the expected section entries",
@@ -78,6 +78,28 @@ namespace UE::ToolMenus::Private
 		Describe("Partial Match", [this]()
 		{
 			It("Should match any existing section", [this]()
+			{
+				const FMenu ExpectedPartialMenuStructure(
+					"TestMenu",
+					{
+						Section("MiddleSection"),
+						Section("SectionThatDoesntExist")
+					});
+
+				const UToolMenu* ActualMenu = GenerateTestMenu();
+
+				FToolMenuAutomationTestAdapter::FToolMenuAutomationTestAdapterParameters TestAdapterParameters;
+				TestAdapterParameters.SectionMatchParameters.bActualHasAllExpectedItems = false;
+				TestAdapterParameters.EntryMatchParameters.bActualHasAllExpectedItems = false;
+
+				TestMatch(
+					"Actual menu sections to (partially) match the expected section(s)",
+					ExpectedPartialMenuStructure,
+					*ActualMenu,
+					TestAdapterParameters);
+			});
+
+			It("Should allow extra sections", [this]()
 			{
 				const FMenu ExpectedPartialMenuStructure(
 					"TestMenu",
@@ -88,8 +110,8 @@ namespace UE::ToolMenus::Private
 				const UToolMenu* ActualMenu = GenerateTestMenu();
 
 				FToolMenuAutomationTestAdapter::FToolMenuAutomationTestAdapterParameters TestAdapterParameters;
-				TestAdapterParameters.SectionMatchParameters.bHasAllItems = false;
-				TestAdapterParameters.EntryMatchParameters.bHasAllItems = false;
+				TestAdapterParameters.SectionMatchParameters.bActualHasAllExpectedItems = true;
+				TestAdapterParameters.EntryMatchParameters.bActualHasAllExpectedItems = false;
 
 				TestMatch(
 					"Actual menu sections to (partially) match the expected section(s)",
@@ -110,9 +132,9 @@ namespace UE::ToolMenus::Private
 
 				const UToolMenu* ActualMenu = GenerateTestMenu();
 				FToolMenuAutomationTestAdapter::FToolMenuAutomationTestAdapterParameters TestAdapterParameters;
-				TestAdapterParameters.SectionMatchParameters.bHasAllItems = false;
-				TestAdapterParameters.SectionMatchParameters.bHasSameOrder = true;
-				TestAdapterParameters.EntryMatchParameters.bHasAllItems = false;
+				TestAdapterParameters.SectionMatchParameters.bActualHasAllExpectedItems = true; // otherwise the out-of-order section would be considered extra, and valid
+				TestAdapterParameters.SectionMatchParameters.bActualHasExpectedOrder = true;
+				TestAdapterParameters.EntryMatchParameters.bActualHasAllExpectedItems = false;
 
 				TestMismatch(
 					"Menu sections shouldn't match wrongly ordered sections",
@@ -135,8 +157,8 @@ namespace UE::ToolMenus::Private
 				const UToolMenu* ActualMenu = GenerateTestMenu();
 
 				FToolMenuAutomationTestAdapter::FToolMenuAutomationTestAdapterParameters TestAdapterParameters;
-				TestAdapterParameters.SectionMatchParameters.bHasAllItems = false;
-				TestAdapterParameters.EntryMatchParameters.bHasAllItems = false;
+				TestAdapterParameters.SectionMatchParameters.bActualHasAllExpectedItems = false;
+				TestAdapterParameters.EntryMatchParameters.bActualHasAllExpectedItems = false;
 
 				TestMatch(
 					"Actual menu entries to (partially) match the expected entries(s)",
@@ -159,9 +181,9 @@ namespace UE::ToolMenus::Private
 				const UToolMenu* ActualMenu = GenerateTestMenu();
 
 				FToolMenuAutomationTestAdapter::FToolMenuAutomationTestAdapterParameters TestAdapterParameters;
-				TestAdapterParameters.SectionMatchParameters.bHasAllItems = false;
-				TestAdapterParameters.EntryMatchParameters.bHasAllItems = false;
-				TestAdapterParameters.EntryMatchParameters.bHasSameOrder = true;
+				TestAdapterParameters.SectionMatchParameters.bActualHasAllExpectedItems = true; // otherwise the out-of-order entry would be considered extra, and valid
+				TestAdapterParameters.EntryMatchParameters.bActualHasAllExpectedItems = false;
+				TestAdapterParameters.EntryMatchParameters.bActualHasExpectedOrder = true;
 
 				TestMismatch(
 					"Menu entries shouldn't match wrongly ordered entries",
diff --git a/Engine/Source/Developer/ToolMenus/Public/Tests/ToolMenusTestUtilities.h b/Engine/Source/Developer/ToolMenus/Public/Tests/ToolMenusTestUtilities.h
--- a/Engine/Source/Developer/ToolMenus/Public/Tests/ToolMenusTestUtilities.h
+++ b/Engine/Source/Developer/ToolMenus/Public/Tests/ToolMenusTestUtilities.h
@@ -9,12 +9,11 @@
 
 namespace UE::ToolMenus
 {
-	// @todo: account for extension points, ie. allow for other arbitrary entries, but enforce order of known entries to test against
-
 	enum class EExpectedOccurrence : uint8
 	{
-		Any,
-		ExactlyZero,
+		Any,			// Zero or more
+		All,			// Exact match
+		ExactlyZero,	// Should not match at all
 		ExactlyOne,
 		OneOrMore
 	};
@@ -23,16 +22,22 @@ namespace UE::ToolMenus
 	 * @note: that these parameters apply to the item type, not the children of that item (with some exceptions).
 	 * When used for sections, it applies to Sections in general and not entries within that section.
 	 */
-	struct TOOLMENUS_API FMenuMatchParameters
+	struct FMenuMatchParameters
 	{
-		/** All expected entries/sections must exist in the actual menu. */
-		bool bHasAllItems = false;
+		/**
+		 * If true, all expected entries/sections must exist in the actual menu.
+		 * If false, only one or more expected items are expected to be found, and the test will pass even if additional items are expected.
+		 */
+		bool bActualHasAllExpectedItems = false;
 
 		/** The expected entries/sections must be in the same order they are in the actual menu. */
-		bool bHasSameOrder = true;
+		bool bActualHasExpectedOrder = true;
 
-		/** Allow extra entries in the actual menu, skipping over them if they're found in actual but not expected. Relevant when bHasAllItems is true. */
-		bool bAllowExtraItems = true;
+		/**
+		 * If true, all of the actual items must be found in the expected menu.
+		 * If false, only one or more actual items are expected to be found, and the test will pass even if additional items are found in the actual menu.
+		 */
+		bool bExpectedHasAllActualItems = false;
 
 		/** Attempts to match children if any are expected, otherwise it skips them. */
 		bool bMatchChildrenIfAnyExpected = true;
@@ -211,7 +216,7 @@ namespace UE::ToolMenus
 
 	inline TSharedRef<FMenuEntry> Separator()
 	{
-		TSharedRef<FMenuEntry> Result = MakeShared<FMenuEntry>(NAME_All);
+		TSharedRef<FMenuEntry> Result = MakeShared<FMenuEntry>(NAME_None);
 		Result->Type = EMultiBlockType::Separator;
 		return Result;
 	}
@@ -250,8 +255,8 @@ namespace UE::ToolMenus
 		struct FToolMenuAutomationTestAdapterParameters
 		{
 			FToolMenuAutomationTestAdapterParameters()
-				: SectionMatchParameters({ true, true, true, true })
-				, EntryMatchParameters({ true, true, true, true })
+				: SectionMatchParameters({ true, true, false, true })
+				, EntryMatchParameters({ true, true, false, true })
 			{
 
 			}
