Implement the new viewport toolbar for the Dataflow Editor.
- Implement UToolMenu per-block visibility for toolbar-level items.
- Change the default of IPreviewLODController::GetAutoLODStartingIndex() to include the "auto" lod behavior.

diff --git a/Engine/Plugins/Experimental/Dataflow/Source/DataflowEditor/Private/Dataflow/DataflowConstructionViewport.cpp b/Engine/Plugins/Experimental/Dataflow/Source/DataflowEditor/Private/Dataflow/DataflowConstructionViewport.cpp
--- a/Engine/Plugins/Experimental/Dataflow/Source/DataflowEditor/Private/Dataflow/DataflowConstructionViewport.cpp
+++ b/Engine/Plugins/Experimental/Dataflow/Source/DataflowEditor/Private/Dataflow/DataflowConstructionViewport.cpp
@@ -2,6 +2,7 @@
 
 #include "Dataflow/DataflowConstructionViewport.h"
 
+#include "DataflowSceneProfileIndexStorage.h"
 #include "Dataflow/DataflowActor.h"
 #include "Dataflow/DataflowEditorCommands.h"
 #include "Dataflow/DataflowEditorMode.h"
@@ -14,6 +15,9 @@
 #include "Dataflow/DataflowEditorPreviewSceneBase.h"
 #include "Dataflow/DataflowRenderingViewMode.h"
 #include "Dataflow/DataflowSimulationPanel.h"
+#include "ToolMenus.h"
+#include "Dataflow/DataflowConstructionVisualization.h"
+#include "ViewportToolbar/UnrealEdViewportToolbar.h"
 
 #define LOCTEXT_NAMESPACE "SDataflowConstructionViewport"
 
@@ -49,6 +53,128 @@ TSharedPtr<SWidget> SDataflowConstructionViewport::MakeViewportToolbar()
 		.CommandList(CommandList);
 }
 
+TSharedPtr<SWidget> SDataflowConstructionViewport::BuildViewportToolbar()
+{
+	const FName ToolbarName = "Dataflow.ConstructionViewportToolbar";
+	
+	if (!UToolMenus::Get()->IsMenuRegistered(ToolbarName))
+	{
+		UToolMenu* Menu = UToolMenus::Get()->RegisterMenu(ToolbarName, NAME_None, EMultiBoxType::SlimHorizontalToolBar);
+		Menu->StyleName = "ViewportToolbar";
+		
+		Menu->FindOrAddSection("Left");
+		
+		{
+			FToolMenuSection& RightSection = Menu->AddSection("Right");
+			RightSection.Alignment = EToolMenuSectionAlign::Last;
+			
+			RightSection.AddEntry(UE::UnrealEd::CreateViewModesSubmenu());
+			
+			{
+				RightSection.AddDynamicEntry("DynamicShowAndCamera", FNewToolMenuSectionDelegate::CreateLambda([](FToolMenuSection& Section)
+				{
+					UUnrealEdViewportToolbarContext* Context = Section.FindContext<UUnrealEdViewportToolbarContext>();
+					if (!Context)
+					{
+						return;
+					}
+					
+					TSharedPtr<SDataflowConstructionViewport> ConstructionViewport = StaticCastSharedPtr<SDataflowConstructionViewport>(Context->Viewport.Pin());
+					if (!ConstructionViewport)
+					{
+						return;
+					}
+				
+					{
+						// Show Menu
+						FToolMenuEntry& Entry = Section.AddSubMenu(
+							"Show",
+							LOCTEXT("ShowSubmenuLabel", "Show"),
+							LOCTEXT("ShowSubmenuTooltip", "Display options for the construction viewport."),
+							FNewMenuDelegate::CreateLambda([WeakViewport = ConstructionViewport.ToWeakPtr()](FMenuBuilder& Menu)
+							{
+								if (const TSharedPtr<SDataflowConstructionViewport> PinnedEditorViewport = WeakViewport.Pin())
+								{
+									if (const TSharedPtr<FEditorViewportClient> ViewportClient = PinnedEditorViewport->GetViewportClient())
+									{
+										for (const TPair<FName, TUniquePtr<UE::Dataflow::IDataflowConstructionVisualization>>& Visualization : UE::Dataflow::FDataflowConstructionVisualizationRegistry::GetInstance().GetVisualizations())
+										{
+											Visualization.Value->ExtendViewportShowMenu(StaticCastSharedPtr<FDataflowConstructionViewportClient>(ViewportClient), Menu);
+										}
+									}
+								}
+							})
+						);
+						Entry.Icon = FSlateIcon(FAppStyle::GetAppStyleSetName(), "Level.VisibleHighlightIcon16x");
+						Entry.ToolBarData.LabelOverride = FText();
+						Entry.ToolBarData.ResizeParams.ClippingPriority = 800;
+					}
+					
+					{
+						
+						// Camera / View Options
+						TAttribute<FText> Label = TAttribute<FText>::CreateLambda([WeakViewport = ConstructionViewport.ToWeakPtr()]
+						{
+							if (const TSharedPtr<SDataflowConstructionViewport> Viewport = WeakViewport.Pin())
+							{
+								if (UDataflowEditorMode* EditorMode = Viewport->GetEdMode())
+								{
+									return EditorMode->GetConstructionViewMode()->GetButtonText();
+								}
+							}
+							
+							return LOCTEXT("DataflowConstructionViewMenuTitle_Default", "View");
+						});
+						
+						Section.AddSubMenu(
+							"Camera",
+							Label,
+							LOCTEXT("CameraSubMenuTooltip", "Display options for the construction viewport."),
+							FNewToolMenuDelegate::CreateLambda([WeakViewport = ConstructionViewport.ToWeakPtr()](UToolMenu* Menu)
+							{
+								FToolMenuSection& SimulationSection = Menu->AddSection("Simulation", LOCTEXT("SimulationSection", "Simulation"));
+							
+								if (const TSharedPtr<SDataflowConstructionViewport> Viewport = WeakViewport.Pin())
+								{
+									if (UDataflowEditorMode* EditorMode = Viewport->GetEdMode())
+									{
+										for (const TPair<FName, TSharedPtr<FUICommandInfo>>& NameAndCommand : FDataflowEditorCommandsImpl::Get().SetConstructionViewModeCommands)
+										{
+											if (EditorMode->CanChangeConstructionViewModeTo(NameAndCommand.Key))
+											{
+												SimulationSection.AddMenuEntry(NameAndCommand.Value);
+											}
+										}
+									}
+								}
+							}),
+							false,
+							FSlateIcon(FAppStyle::Get().GetStyleSetName(), "EditorViewport.Perspective")
+						);
+					}
+				}));
+			}
+		}
+	}
+	
+	FToolMenuContext Context;
+	{
+		Context.AppendCommandList(GetCommandList());
+		Context.AddExtender(GetExtenders());
+		
+		UUnrealEdViewportToolbarContext* ContextObject = UE::UnrealEd::CreateViewportToolbarDefaultContext(SharedThis(this));
+		Context.AddObject(ContextObject);
+	}
+	
+	return UToolMenus::Get()->GenerateWidget(ToolbarName, Context);
+}
+
+TSharedPtr<IPreviewProfileController> SDataflowConstructionViewport::CreatePreviewProfileController()
+{
+	TSharedPtr<FDataflowConstructionSceneProfileIndexStorage> ProfileIndexStorage = MakeShared<FDataflowConstructionSceneProfileIndexStorage>(GetConstructionScene());
+	return MakeShared<FDataflowPreviewProfileController>(ProfileIndexStorage);
+}
+
 void SDataflowConstructionViewport::OnFocusViewportToSelection()
 {
 	const UDataflowEditorMode* const DataflowEdMode = GetEdMode();
diff --git a/Engine/Plugins/Experimental/Dataflow/Source/DataflowEditor/Private/Dataflow/DataflowEditorCommands.cpp b/Engine/Plugins/Experimental/Dataflow/Source/DataflowEditor/Private/Dataflow/DataflowEditorCommands.cpp
--- a/Engine/Plugins/Experimental/Dataflow/Source/DataflowEditor/Private/Dataflow/DataflowEditorCommands.cpp
+++ b/Engine/Plugins/Experimental/Dataflow/Source/DataflowEditor/Private/Dataflow/DataflowEditorCommands.cpp
@@ -70,10 +70,10 @@ void FDataflowEditorCommandsImpl::RegisterCommands()
 	UI_COMMAND(ZoomToFitGraph, "ZoomToFitGraph", "Fit the graph in the graph editor viewport.", EUserInterfaceActionType::None, FInputChord(EKeys::F));
 
 	UI_COMMAND(AddWeightMapNode, "Add Weight Map", "Paint weight maps on the mesh", EUserInterfaceActionType::Button, FInputChord());
-	UI_COMMAND(RebuildSimulationScene, "RebuildSimulationScene", "Rebuild the simulation scene", EUserInterfaceActionType::ToggleButton, FInputChord());
-	UI_COMMAND(PauseSimulationScene, "PauseSimulationScene", "Pause the simulation scene", EUserInterfaceActionType::ToggleButton, FInputChord());
-	UI_COMMAND(StartSimulationScene, "StartSimulationScene", "Start the simulation scene", EUserInterfaceActionType::ToggleButton, FInputChord());
-	UI_COMMAND(StepSimulationScene, "StepSimulationScene", "Step the simulation scene", EUserInterfaceActionType::ToggleButton, FInputChord());
+	UI_COMMAND(RebuildSimulationScene, "Rebuild Simulation", "Rebuild the simulation scene", EUserInterfaceActionType::Button, FInputChord());
+	UI_COMMAND(PauseSimulationScene, "Pause Simulation", "Pause the simulation scene", EUserInterfaceActionType::Button, FInputChord());
+	UI_COMMAND(StartSimulationScene, "Start Simulation", "Start the simulation scene", EUserInterfaceActionType::Button, FInputChord());
+	UI_COMMAND(StepSimulationScene, "Step Simulation", "Step the simulation scene", EUserInterfaceActionType::Button, FInputChord());
 
 	UI_COMMAND(AddNewVariable, "Variable", "Adds a new variable to this dataflow graph.", EUserInterfaceActionType::Button, FInputChord());
 	UI_COMMAND(AddNewSubGraph, "SubGraph", "Adds a new sub-graph to this dataflow graph.", EUserInterfaceActionType::Button, FInputChord());
@@ -92,7 +92,7 @@ void FDataflowEditorCommandsImpl::RegisterCommands()
 			ViewMode->GetButtonText(),
 			ViewMode->GetTooltipText(),
 			FSlateIcon(),
-			EUserInterfaceActionType::Button,
+			EUserInterfaceActionType::RadioButton,
 			FInputChord()
 		);
 		SetConstructionViewModeCommands.Add(ViewMode->GetName(), SetViewModeCommand);
diff --git a/Engine/Plugins/Experimental/Dataflow/Source/DataflowEditor/Private/Dataflow/DataflowSimulationScene.cpp b/Engine/Plugins/Experimental/Dataflow/Source/DataflowEditor/Private/Dataflow/DataflowSimulationScene.cpp
--- a/Engine/Plugins/Experimental/Dataflow/Source/DataflowEditor/Private/Dataflow/DataflowSimulationScene.cpp
+++ b/Engine/Plugins/Experimental/Dataflow/Source/DataflowEditor/Private/Dataflow/DataflowSimulationScene.cpp
@@ -156,6 +156,15 @@ void FDataflowSimulationScene::StepSimulationScene() const
 	}
 }
 
+bool FDataflowSimulationScene::IsSimulationEnabled() const
+{
+	if(SceneDescription && (SceneDescription->CacheAsset == nullptr))
+	{
+		return GetWorld()->GetSubsystem<UDataflowSimulationManager>()->GetSimulationEnabled();
+	}
+	return false;
+}
+
 void FDataflowSimulationScene::RebuildSimulationScene(const bool bIsSimulationEnabled)
 {
 	if(SceneDescription && (SceneDescription->CacheAsset == nullptr))
diff --git a/Engine/Plugins/Experimental/Dataflow/Source/DataflowEditor/Private/Dataflow/DataflowSimulationViewport.cpp b/Engine/Plugins/Experimental/Dataflow/Source/DataflowEditor/Private/Dataflow/DataflowSimulationViewport.cpp
--- a/Engine/Plugins/Experimental/Dataflow/Source/DataflowEditor/Private/Dataflow/DataflowSimulationViewport.cpp
+++ b/Engine/Plugins/Experimental/Dataflow/Source/DataflowEditor/Private/Dataflow/DataflowSimulationViewport.cpp
@@ -1,6 +1,9 @@
 // Copyright Epic Games, Inc. All Rights Reserved.
 
 #include "Dataflow/DataflowSimulationViewport.h"
+
+#include "DataflowPreviewProfileController.h"
+#include "DataflowSceneProfileIndexStorage.h"
 #include "Dataflow/DataflowActor.h"
 #include "Dataflow/DataflowEditorMode.h"
 #include "Dataflow/DataflowSimulationViewportClient.h"
@@ -13,6 +16,9 @@
 #include "Dataflow/DataflowSimulationPanel.h"
 #include "Dataflow/DataflowSimulationViewportToolbar.h"
 #include "Dataflow/DataflowSimulationVisualization.h"
+#include "ToolMenus.h"
+#include "LODSyncInterface.h"
+#include "ViewportToolbar/UnrealEdViewportToolbar.h"
 #include "Widgets/Text/SRichTextBlock.h"
 
 #define LOCTEXT_NAMESPACE "SDataflowSimulationViewport"
@@ -118,6 +124,149 @@ TSharedPtr<SWidget> SDataflowSimulationViewport::MakeViewportToolbar()
 	return SNew(SDataflowSimulationViewportToolBar, SharedThis(this)).CommandList(CommandList);
 }
 
+TSharedPtr<SWidget> SDataflowSimulationViewport::BuildViewportToolbar()
+{
+	const FName ToolbarName = "Dataflow.SimulationViewportToolbar";
+	
+	if (!UToolMenus::Get()->IsMenuRegistered(ToolbarName))
+	{
+		UToolMenu* Menu = UToolMenus::Get()->RegisterMenu(ToolbarName, NAME_None, EMultiBoxType::SlimHorizontalToolBar);
+		Menu->StyleName = "ViewportToolbar";
+		
+		Menu->FindOrAddSection("Left");
+		
+		{
+			FToolMenuSection& RightSection = Menu->AddSection("Right");
+			RightSection.Alignment = EToolMenuSectionAlign::Last;
+			
+			RightSection.AddEntry(UE::UnrealEd::CreateCameraSubmenu(UE::UnrealEd::FViewportCameraMenuOptions().ShowCameraMovement()));
+			RightSection.AddEntry(UE::UnrealEd::CreateViewModesSubmenu());
+			RightSection.AddEntry(UE::UnrealEd::CreateDefaultShowSubmenu());
+			
+			{
+				// LOD
+				RightSection.AddDynamicEntry("DynamicLOD", FNewToolMenuSectionDelegate::CreateLambda([](FToolMenuSection& Section)
+				{
+					if (UUnrealEdViewportToolbarContext* Context = Section.FindContext<UUnrealEdViewportToolbarContext>())
+					{
+						TWeakPtr<SDataflowSimulationViewport> SimulationViewport = StaticCastWeakPtr<SDataflowSimulationViewport>(Context->Viewport);
+						Section.AddEntry(UE::UnrealEd::CreatePreviewLODSelectionSubmenu(SimulationViewport));
+					}
+				}));
+			}
+			
+			{
+				RightSection.AddEntry(UE::UnrealEd::CreateAssetViewerProfileSubmenu());
+				UToolMenu* AssetViewerMenu = UToolMenus::Get()->ExtendMenu(UToolMenus::JoinMenuPaths(ToolbarName, "AssetViewerProfile"));
+				FToolMenuSection& SimulationSection = AssetViewerMenu->FindOrAddSection("SimulationControls", LOCTEXT("SimulationControlsSection", "Simulation Scene Controls"));
+				
+				SimulationSection.AddDynamicEntry("DynamicSceneControls", FNewToolMenuSectionDelegate::CreateLambda([](FToolMenuSection& Section)
+				{
+					UUnrealEdViewportToolbarContext* Context = Section.FindContext<UUnrealEdViewportToolbarContext>();
+					if (!Context)
+					{
+						return;
+					}
+					
+					TWeakPtr<SDataflowSimulationViewport> WeakViewport = StaticCastWeakPtr<SDataflowSimulationViewport>(Context->Viewport);
+					if (!WeakViewport.IsValid())
+					{
+						return;
+					}
+				
+					// Simulation Controls
+					FToolMenuEntryToolBarData ToolBarData;
+					ToolBarData.BlockGroupName = "SimulationControlGroup";
+					ToolBarData.LabelOverride = FText::GetEmpty();
+					ToolBarData.ResizeParams.ClippingPriority = 2000;
+					
+					FToolMenuEntry& Rebuild = Section.AddMenuEntry(FDataflowEditorCommands::Get().RebuildSimulationScene);
+					Rebuild.ToolBarData = ToolBarData;
+					Rebuild.Icon = FSlateIcon(FAppStyle::Get().GetStyleSetName(), "Animation.Backward_End");
+					Rebuild.SetShowInToolbarTopLevel(true);
+					
+					FToolMenuEntry& Pause = Section.AddMenuEntry(FDataflowEditorCommands::Get().PauseSimulationScene);
+					Pause.ToolBarData = ToolBarData;
+					Pause.Icon = FSlateIcon(FAppStyle::Get().GetStyleSetName(), "Animation.Pause");
+					Pause.SetShowInToolbarTopLevel(true);
+					Pause.SetVisibility(TAttribute<EVisibility>::CreateLambda([WeakViewport]
+					{
+						if (TSharedPtr<SDataflowSimulationViewport> Viewport = WeakViewport.Pin())
+						{
+							if (TSharedPtr<FEditorViewportClient> Client = Viewport->GetViewportClient())
+							{
+								if (FDataflowSimulationScene* SimulationScene = StaticCast<FDataflowSimulationScene*>(Client->GetPreviewScene()))
+								{
+									return SimulationScene->IsSimulationEnabled() ? EVisibility::Visible : EVisibility::Collapsed;
+								}
+							}
+						}
+						return EVisibility::Visible;
+					}));
+					
+					FToolMenuEntry& Start = Section.AddMenuEntry(FDataflowEditorCommands::Get().StartSimulationScene);
+					Start.ToolBarData = ToolBarData;
+					Start.Icon = FSlateIcon(FAppStyle::Get().GetStyleSetName(), "Icons.Play");
+					Start.SetShowInToolbarTopLevel(true);
+					Start.SetVisibility(TAttribute<EVisibility>::CreateLambda([WeakViewport]
+                    {
+                    	if (TSharedPtr<SDataflowSimulationViewport> Viewport = WeakViewport.Pin())
+                    	{
+                    		if (TSharedPtr<FEditorViewportClient> Client = Viewport->GetViewportClient())
+                    		{
+                    			if (FDataflowSimulationScene* SimulationScene = StaticCast<FDataflowSimulationScene*>(Client->GetPreviewScene()))
+                    			{
+                    				return SimulationScene->IsSimulationEnabled() ? EVisibility::Collapsed : EVisibility::Visible;
+                    			}
+                    		}
+                    	}
+                    	return EVisibility::Visible;
+                    }));
+					
+					FToolMenuEntry& Step = Section.AddMenuEntry(FDataflowEditorCommands::Get().StepSimulationScene);
+					Step.ToolBarData = ToolBarData;
+					Step.Icon = FSlateIcon(FAppStyle::Get().GetStyleSetName(), "Animation.Forward_Step");
+					Step.SetShowInToolbarTopLevel(true);
+				}));	
+			}
+		}
+	}
+	
+	FToolMenuContext Context;
+	{
+		Context.AppendCommandList(GetCommandList());
+		Context.AddExtender(GetExtenders());
+		
+		UUnrealEdViewportToolbarContext* ContextObject = UE::UnrealEd::CreateViewportToolbarDefaultContext(SharedThis(this));
+		Context.AddObject(ContextObject);
+	}
+	
+	return UToolMenus::Get()->GenerateWidget(ToolbarName, Context);
+}
+
+TSharedPtr<IPreviewProfileController> SDataflowSimulationViewport::CreatePreviewProfileController()
+{
+	TSharedPtr<FDataflowSimulationSceneProfileIndexStorage> ProfileIndexStorage;
+
+	if (const TSharedPtr<FDataflowSimulationViewportClient> ViewportClient = StaticCastSharedPtr<FDataflowSimulationViewportClient>(GetViewportClient()))
+	{
+		if (ViewportClient->GetDataflowEditorToolkit().IsValid())
+		{
+			if (const TSharedPtr<FDataflowEditorToolkit> Toolkit = ViewportClient->GetDataflowEditorToolkit().Pin())
+			{
+				ProfileIndexStorage = Toolkit->GetSimulationSceneProfileIndexStorage();
+			}
+		}
+	}
+	
+	if (!ProfileIndexStorage)
+	{
+		ProfileIndexStorage = MakeShared<FDataflowSimulationSceneProfileIndexStorage>(GetSimulationScene().Get());
+	}
+	
+	return MakeShared<FDataflowPreviewProfileController>(ProfileIndexStorage);
+}
+
 void SDataflowSimulationViewport::OnFocusViewportToSelection()
 {
 	if(const FDataflowPreviewSceneBase* PreviewScene = static_cast<FDataflowPreviewSceneBase*>(Client->GetPreviewScene()))
@@ -215,6 +364,51 @@ void SDataflowSimulationViewport::OnFloatingButtonClicked()
 {
 }
 
+int32 SDataflowSimulationViewport::GetCurrentLOD() const
+{
+	if (TSharedPtr<FDataflowSimulationScene> SimScene = GetSimulationScene())
+	{
+		return SimScene->GetPreviewLOD();
+	}
+	return -1;
+}
+
+int32 SDataflowSimulationViewport::GetLODCount() const
+{
+	int32 MaxNumLODs = 0;
+	if (const TSharedPtr<const FDataflowSimulationScene>& SimScene = GetSimulationScene())
+	{
+		if (const AActor* const PreviewActor = SimScene->GetPreviewActor())
+		{
+			PreviewActor->ForEachComponent<UActorComponent>(true, [&MaxNumLODs](UActorComponent* Component)
+			{
+				if (const ILODSyncInterface* const LODInterface = Cast<ILODSyncInterface>(Component))
+				{
+					MaxNumLODs = FMath::Max(MaxNumLODs, LODInterface->GetNumSyncLODs());
+				}
+			});
+		}
+	}
+	return MaxNumLODs;
+}
+
+bool SDataflowSimulationViewport::IsLODSelected(int32 LODIndex) const
+{
+	if (TSharedPtr<FDataflowSimulationScene> SimScene = GetSimulationScene())
+	{
+		return SimScene->GetPreviewLOD() == LODIndex;
+	}
+	return false;
+}
+
+void SDataflowSimulationViewport::SetLODLevel(int32 LODIndex)
+{
+	if (TSharedPtr<FDataflowSimulationScene> SimScene = GetSimulationScene())
+	{
+		return SimScene->SetPreviewLOD(LODIndex);
+	}
+}
+
 float SDataflowSimulationViewport::GetViewMinInput() const
 {
 	return 0.0f;
diff --git a/Engine/Plugins/Experimental/Dataflow/Source/DataflowEditor/Public/Dataflow/DataflowConstructionViewport.h b/Engine/Plugins/Experimental/Dataflow/Source/DataflowEditor/Public/Dataflow/DataflowConstructionViewport.h
--- a/Engine/Plugins/Experimental/Dataflow/Source/DataflowEditor/Public/Dataflow/DataflowConstructionViewport.h
+++ b/Engine/Plugins/Experimental/Dataflow/Source/DataflowEditor/Public/Dataflow/DataflowConstructionViewport.h
@@ -29,6 +29,8 @@ public:
 	// SEditorViewport
 	virtual void BindCommands() override;
 	virtual TSharedPtr<SWidget> MakeViewportToolbar() override;
+	virtual TSharedPtr<SWidget> BuildViewportToolbar() override;
+	virtual TSharedPtr<IPreviewProfileController> CreatePreviewProfileController() override;
 	virtual bool IsVisible() const override;
 	virtual void OnFocusViewportToSelection() override;
 
diff --git a/Engine/Plugins/Experimental/Dataflow/Source/DataflowEditor/Public/Dataflow/DataflowSimulationScene.h b/Engine/Plugins/Experimental/Dataflow/Source/DataflowEditor/Public/Dataflow/DataflowSimulationScene.h
--- a/Engine/Plugins/Experimental/Dataflow/Source/DataflowEditor/Public/Dataflow/DataflowSimulationScene.h
+++ b/Engine/Plugins/Experimental/Dataflow/Source/DataflowEditor/Public/Dataflow/DataflowSimulationScene.h
@@ -124,6 +124,9 @@ public:
 
 	/** Step the simulation */
 	void StepSimulationScene() const;
+	
+	/** Gets whether the simulation is running */
+	bool IsSimulationEnabled() const;
 
 	/** Rebuild the simulation scene */
 	void RebuildSimulationScene(const bool bIsSimulationEnabled);
diff --git a/Engine/Plugins/Experimental/Dataflow/Source/DataflowEditor/Public/Dataflow/DataflowSimulationViewport.h b/Engine/Plugins/Experimental/Dataflow/Source/DataflowEditor/Public/Dataflow/DataflowSimulationViewport.h
--- a/Engine/Plugins/Experimental/Dataflow/Source/DataflowEditor/Public/Dataflow/DataflowSimulationViewport.h
+++ b/Engine/Plugins/Experimental/Dataflow/Source/DataflowEditor/Public/Dataflow/DataflowSimulationViewport.h
@@ -2,6 +2,7 @@
 #pragma once
 
 #include "CoreMinimal.h"
+#include "IPreviewLODController.h"
 #include "SBaseCharacterFXEditorViewport.h"
 #include "SCommonEditorViewportToolbarBase.h"
 
@@ -12,7 +13,7 @@ class FDataflowConstructionViewportClient;
 
 // ----------------------------------------------------------------------------------
 
-class SDataflowSimulationViewport : public SBaseCharacterFXEditorViewport, public ICommonEditorViewportToolbarInfoProvider
+class SDataflowSimulationViewport : public SBaseCharacterFXEditorViewport, public ICommonEditorViewportToolbarInfoProvider, public IPreviewLODController
 {
 public:
 	SLATE_BEGIN_ARGS(SDataflowSimulationViewport) {}
@@ -28,12 +29,20 @@ public:
 	virtual void BindCommands() override;
 	virtual bool IsVisible() const override;
 	virtual TSharedPtr<SWidget> MakeViewportToolbar() override;
+	virtual TSharedPtr<SWidget> BuildViewportToolbar() override;
+	virtual TSharedPtr<IPreviewProfileController> CreatePreviewProfileController() override;
 	virtual void OnFocusViewportToSelection() override;
-
+	
 	// ICommonEditorViewportToolbarInfoProvider
 	virtual TSharedRef<class SEditorViewport> GetViewportWidget() override;
 	virtual TSharedPtr<FExtender> GetExtenders() const override;
 	virtual void OnFloatingButtonClicked() override;
+	
+	// IPreviewLODController
+	virtual int32 GetCurrentLOD() const override;
+	virtual int32 GetLODCount() const override;
+	virtual bool IsLODSelected(int32 LODIndex) const override;
+	virtual void SetLODLevel(int32 LODIndex) override;
 
 	/** Get the simulation scene */
 	const TSharedPtr<class FDataflowSimulationScene>& GetSimulationScene() const;
diff --git a/Engine/Source/Developer/ToolMenus/Private/ToolMenus.cpp b/Engine/Source/Developer/ToolMenus/Private/ToolMenus.cpp
--- a/Engine/Source/Developer/ToolMenus/Private/ToolMenus.cpp
+++ b/Engine/Source/Developer/ToolMenus/Private/ToolMenus.cpp
@@ -34,6 +34,40 @@ UToolMenus* CreateToolMenusInstance()
 	return Instance;
 }
 
+/** Combine two visibility attributes with the first attribute's non-visibility taking precedence. */
+TAttribute<EVisibility> CombineVisibilityAttributes(const TAttribute<EVisibility>& PrimaryVisibility, const TAttribute<EVisibility>& SecondaryVisibility)
+{
+	if (PrimaryVisibility.IsSet() && SecondaryVisibility.IsSet())
+	{
+		return TAttribute<EVisibility>::CreateLambda(
+			[Primary = PrimaryVisibility, Secondary = SecondaryVisibility]()
+			{
+				// If section visibility is set, and section is not visible, return Section Visibility
+				const EVisibility PrimaryVisibility = Primary.Get();
+				if (PrimaryVisibility != EVisibility::Visible)
+				{
+					return PrimaryVisibility;
+				}
+
+				// Stick to this Block visibility
+				return Secondary.Get();
+			}
+		);
+	}
+		
+	if (PrimaryVisibility.IsSet())
+	{
+		return PrimaryVisibility;
+	}
+		
+	if (SecondaryVisibility.IsSet())
+	{
+		return SecondaryVisibility;
+	}
+		
+	return TAttribute<EVisibility>();
+}
+
 } // namespace UE::ToolMenus::Private
 
 UToolMenus* UToolMenus::Singleton = nullptr;
@@ -93,25 +127,6 @@ FToolUIActionChoice::FToolUIActionChoice(const TSharedPtr< const FUICommandInfo
 
 class FPopulateMenuBuilderWithToolMenuEntry
 {
-private:
-	// Returns visibility for the specified Entry + Section
-	static TAttribute<EVisibility> GetEntryVisibility(FToolMenuEntry& InMenuEntry, FToolMenuSection& InSection)
-	{
-		return TAttribute<EVisibility>::CreateLambda(
-			[BlockVisibility = InMenuEntry.Visibility, SectionVisibility = InSection.Visibility]()
-			{
-				// If section visibility is set, and section is not visible, return Section Visibility
-				if (SectionVisibility.IsSet() && SectionVisibility.Get() != EVisibility::Visible)
-				{
-					return SectionVisibility.Get();
-				}
-
-				// Stick to this Block visibility
-				return BlockVisibility.Get();
-			}
-		);
-	}
-
 public:
 	FPopulateMenuBuilderWithToolMenuEntry(FMenuBuilder& InMenuBuilder, UToolMenu* InMenuData, FToolMenuSection& InSection, FToolMenuEntry& InBlock, bool bInAllowSubMenuCollapse) :
 		MenuBuilder(InMenuBuilder),
@@ -196,7 +211,7 @@ public:
 
 		if (!bSubMenuAdded)
 		{
-			TAttribute<EVisibility> VisibilityOverride = GetEntryVisibility(Block, Section);
+			const TAttribute<EVisibility> VisibilityOverride = UE::ToolMenus::Private::CombineVisibilityAttributes(Section.Visibility, Block.Visibility);
 
 			if (Widget.IsValid())
 			{
@@ -254,7 +269,7 @@ public:
 
 	void AddStandardEntryToMenuBuilder()
 	{
-		TAttribute<EVisibility> VisibilityOverride = GetEntryVisibility(Block, Section);
+		const TAttribute<EVisibility> VisibilityOverride = UE::ToolMenus::Private::CombineVisibilityAttributes(Section.Visibility, Block.Visibility);
 
 		// First, check for a ToolUIAction, otherwise do the rest of this (have CommandList and Command)
 		// Need another variable to store if we are using a keybind from a command
@@ -347,7 +362,7 @@ public:
 			return;
 		}
 
-		TAttribute<EVisibility> VisibilityOverride = GetEntryVisibility(Block, Section);
+		const TAttribute<EVisibility> VisibilityOverride = UE::ToolMenus::Private::CombineVisibilityAttributes(Section.Visibility, Block.Visibility);
 
 		UIAction = UToolMenus::ConvertUIAction(Block, MenuData->Context);
 		bUIActionIsSet = UIAction.ExecuteAction.IsBound() || UIAction.CanExecuteAction.IsBound() || UIAction.GetActionCheckState.IsBound() || UIAction.IsActionVisibleDelegate.IsBound();
@@ -1498,9 +1513,7 @@ void UToolMenus::PopulateMenuBuilder(FMenuBuilder& MenuBuilder, UToolMenu* MenuD
 	AddReferencedContextObjects(MenuBuilder.GetMultiBox(), MenuData);
 }
 
-void UToolMenus::ExtractChildBlocksFromSubMenu(
-	UToolMenu* ParentMenu, FToolMenuEntry& InBlock, TArray<TTuple<UToolMenu*, FToolMenuEntry*, FName>>& SubMenuBlocks
-)
+void UToolMenus::ExtractChildBlocksFromSubMenu(UToolMenu* ParentMenu, FToolMenuEntry& InBlock, TArray<UE::ToolMenus::FSubBlockReference>& SubMenuBlocks)
 {
 	const bool bBuiltViaDelegate =
 		ConvertWidgetChoice(InBlock.ToolBarData.ComboButtonContextMenuGenerator, ParentMenu->Context).IsBound();
@@ -1528,22 +1541,22 @@ void UToolMenus::ExtractChildBlocksFromSubMenu(
 		for (int j = Section.Blocks.Num() - 1; j >= 0; --j)
 		{
 			FToolMenuEntry& Block = Section.Blocks[j];
-			SubMenuBlocks.Push(TTuple<UToolMenu*, FToolMenuEntry*, FName>(SubMenu, &Block, Section.Name));
+			SubMenuBlocks.Push(UE::ToolMenus::FSubBlockReference(SubMenu, &Section, &Block));
 		}
 	}
 }
 
-void UToolMenus::PopulateToolBarBuilderWithTopLevelChildren(
-	FToolBarBuilder& ToolBarBuilder, UToolMenu* ParentMenu, FToolMenuEntry& InBlock, bool bAddSpaceAfterRaisedChildren
-)
+void UToolMenus::PopulateToolBarBuilderWithTopLevelChildren(FToolBarBuilder& ToolBarBuilder, UToolMenu* ParentMenu, FToolMenuEntry& InBlock, bool bAddSpaceAfterRaisedChildren)
 {
-	TArray<TTuple<UToolMenu*, FToolMenuEntry*, FName>> SubMenuBlocks;
+	using namespace UE::ToolMenus;
+
+	TArray<FSubBlockReference> SubMenuBlocks;
 	// Seed the blocks with the passed-in submenu.
 	ExtractChildBlocksFromSubMenu(ParentMenu, InBlock, SubMenuBlocks);
 
 	// Collect the blocks that might be raised to the top level. Also include separators so we can visualize those in
 	// the toolbar when they appear between two raised blocks.
-	TArray<TTuple<UToolMenu*, FToolMenuEntry*, FName>> BlocksToAdd;
+	TArray<FSubBlockReference> BlocksToAdd;
 
 	// Traverse the sub menu blocks we've found thus far to find more grandchild blocks that are raised (raised set to
 	// boolean true) or could be dynamically raised (raised set to a TAttribute<bool> driven by a delegate) to the
@@ -1551,9 +1564,9 @@ void UToolMenus::PopulateToolBarBuilderWithTopLevelChildren(
 	int32 NumIterations = 0;
 	while (SubMenuBlocks.Num() > 0)
 	{
-		const TTuple<UToolMenu*, FToolMenuEntry*, FName> SubMenuBlock = SubMenuBlocks.Pop(EAllowShrinking::No);
-		UToolMenu* const SubMenu = SubMenuBlock.Get<0>();
-		FToolMenuEntry* const Block = SubMenuBlock.Get<1>();
+		const FSubBlockReference SubMenuBlock = SubMenuBlocks.Pop(EAllowShrinking::No);
+		UToolMenu* const SubMenu = SubMenuBlock.ParentMenu;
+		FToolMenuEntry* const Block = SubMenuBlock.Entry;
 
 		// Keep track of how many blocks we've visited to ensure we don't loop indefinitely.
 		if (++NumIterations > 5000)
@@ -1564,17 +1577,14 @@ void UToolMenus::PopulateToolBarBuilderWithTopLevelChildren(
 			break;
 		}
 
-		const TAttribute<bool> ScriptShowInToolbarTopLevel =
-			Block->ScriptObject ? Block->ScriptObject->CreateShowInToolbarTopLevelAttribute(ParentMenu->Context)
-								: TAttribute<bool>();
+		const TAttribute<bool> ScriptShowInToolbarTopLevel = Block->ScriptObject
+			? Block->ScriptObject->CreateShowInToolbarTopLevelAttribute(ParentMenu->Context)
+			: TAttribute<bool>();
 
 		const bool bIsBound = Block->ShowInToolbarTopLevel.IsBound() || ScriptShowInToolbarTopLevel.IsBound();
-		const bool bIsSetToValue = !bIsBound
-								&& (Block->ShowInToolbarTopLevel.IsSet() || ScriptShowInToolbarTopLevel.IsSet());
-		const bool bIsSetToTrueValue = bIsSetToValue
-									&& (Block->ShowInToolbarTopLevel.Get() || ScriptShowInToolbarTopLevel.Get());
-		const bool bIsBoundOrTrue = bIsBound || bIsSetToTrueValue;
-		if (bIsBoundOrTrue || Block->Type == EMultiBlockType::Separator)
+		const bool bIsSetToValue = !bIsBound && (Block->ShowInToolbarTopLevel.IsSet() || ScriptShowInToolbarTopLevel.IsSet());
+		const bool bIsSetToTrueValue = bIsSetToValue && (Block->ShowInToolbarTopLevel.Get() || ScriptShowInToolbarTopLevel.Get());
+		if (bIsBound || bIsSetToTrueValue || Block->Type == EMultiBlockType::Separator)
 		{
 			BlocksToAdd.Add(SubMenuBlock);
 		}
@@ -1583,13 +1593,13 @@ void UToolMenus::PopulateToolBarBuilderWithTopLevelChildren(
 	}
 
 	// Do not allow leading separators.
-	while (BlocksToAdd.Num() > 0 && BlocksToAdd[0].Get<1>()->Type == EMultiBlockType::Separator)
+	while (BlocksToAdd.Num() > 0 && BlocksToAdd[0].Entry->Type == EMultiBlockType::Separator)
 	{
 		BlocksToAdd.RemoveAt(0);
 	}
 
 	// Do not allow trailing separators.
-	while (BlocksToAdd.Num() > 0 && BlocksToAdd[BlocksToAdd.Num() - 1].Get<1>()->Type == EMultiBlockType::Separator)
+	while (BlocksToAdd.Num() > 0 && BlocksToAdd[BlocksToAdd.Num() - 1].Entry->Type == EMultiBlockType::Separator)
 	{
 		BlocksToAdd.RemoveAt(BlocksToAdd.Num() - 1);
 	}
@@ -1597,8 +1607,8 @@ void UToolMenus::PopulateToolBarBuilderWithTopLevelChildren(
 	// Do not allow rows of separators.
 	for (int i = 1; i < BlocksToAdd.Num(); ++i)
 	{
-		const bool bIsCurrentSeparator = BlocksToAdd[i].Get<1>()->Type == EMultiBlockType::Separator;
-		const bool bWasPreviousSeparator = BlocksToAdd[i - 1].Get<1>()->Type == EMultiBlockType::Separator;
+		const bool bIsCurrentSeparator = BlocksToAdd[i].Entry->Type == EMultiBlockType::Separator;
+		const bool bWasPreviousSeparator = BlocksToAdd[i - 1].Entry->Type == EMultiBlockType::Separator;
 		const bool bPartOfRowOfSeparators = bIsCurrentSeparator && bWasPreviousSeparator;
 
 		if (bPartOfRowOfSeparators)
@@ -1641,19 +1651,19 @@ void UToolMenus::PopulateToolBarBuilderWithTopLevelChildren(
 	TSharedPtr<TArray<TAttribute<bool>>> NextVisibilities = MakeShared<TArray<TAttribute<bool>>>();
 
 	bool bHasRaisedEntrySinceLastSeparator = false;
-	// Seed the previous section with the first blocks's section so we don't start with adding a separator because
+	// Seed the previous section with the first block's section so we don't start with adding a separator because
 	// sections seem to have changed.
-	FName PreviousSectionName = BlocksToAdd[0].Get<2>();
+	FName PreviousSectionName = BlocksToAdd[0].Section->Name;
 	FName PreviousBlockGroupName = NAME_None;
 	for (int i = 0; i < BlocksToAdd.Num(); ++i)
 	{
-		UToolMenu* const SubMenu = BlocksToAdd[i].Get<0>();
-		FToolMenuEntry* const Entry = BlocksToAdd[i].Get<1>();
-		const FName SectionName = BlocksToAdd[i].Get<2>();
+		UToolMenu* const SubMenu = BlocksToAdd[i].ParentMenu;
+		FToolMenuSection* const Section = BlocksToAdd[i].Section;
+		FToolMenuEntry* const Entry = BlocksToAdd[i].Entry;
 
 		// Add a separator if one was found or a new section was encountered.
 		if (bHasRaisedEntrySinceLastSeparator
-			&& (SectionName != PreviousSectionName || Entry->Type == EMultiBlockType::Separator))
+			&& (Section->Name != PreviousSectionName || Entry->Type == EMultiBlockType::Separator))
 		{
 			// Step entry visibility delegate records forward now that we encountered a new separator.
 			PreviousVisibilities.Append(*NextVisibilities);
@@ -1719,7 +1729,7 @@ void UToolMenus::PopulateToolBarBuilderWithTopLevelChildren(
 			PreviousBlockGroupName = EntryBlockGroupName;
 		
 			constexpr bool bRaiseToTopLevel = true;
-			PopulateToolBarBuilderWithEntry(ToolBarBuilder, SubMenu, *Entry, bRaiseToTopLevel);
+			PopulateToolBarBuilderWithEntry(ToolBarBuilder, SubMenu, *Section, *Entry, bRaiseToTopLevel);
 			bHasRaisedEntrySinceLastSeparator = true;
 
 			// Keep track of added entries' visibilities so separators can set their visibility override.
@@ -1731,7 +1741,7 @@ void UToolMenus::PopulateToolBarBuilderWithTopLevelChildren(
 			}
 		}
 
-		PreviousSectionName = SectionName;
+		PreviousSectionName = Section->Name;
 	}
 	
 	if (!PreviousBlockGroupName.IsNone())
@@ -1779,7 +1789,12 @@ void UToolMenus::PopulateToolBarBuilderWithTopLevelChildren(
 }
 
 void UToolMenus::PopulateToolBarBuilderWithEntry(
-	FToolBarBuilder& ToolBarBuilder, UToolMenu* MenuData, FToolMenuEntry& Block, bool bIsRaisingToTopLevel, bool bIsLastBlockOfLastSection
+	FToolBarBuilder& ToolBarBuilder,
+	UToolMenu* MenuData,
+	FToolMenuSection& Section,
+	FToolMenuEntry& Block,
+	bool bIsRaisingToTopLevel,
+	bool bIsLastBlockOfLastSection
 )
 {
 	if (Block.ToolBarData.ConstructLegacy.IsBound())
@@ -1825,21 +1840,46 @@ void UToolMenus::PopulateToolBarBuilderWithEntry(
 		ToolBarBuilder.BeginStyleOverride(OverrideStyleName);
 	}
 
-	TAttribute<EVisibility> VisibilityOverride;
-	if (bIsRaisingToTopLevel && Block.ShowInToolbarTopLevel.IsBound())
+	// If a section is invisible, then nothing is visible
+	TAttribute<EVisibility> Visibility = Section.Visibility;
+	
+	if (bIsRaisingToTopLevel)
 	{
-		VisibilityOverride = TAttribute<EVisibility>::CreateLambda(
-			[ShowInToolbarTopLevel = Block.ShowInToolbarTopLevel]() -> EVisibility
-			{
-				if (ShowInToolbarTopLevel.IsSet() && !ShowInToolbarTopLevel.Get())
+		// Only customize visibility if this is bound.
+		// A static value of `true` got this block here in the first place
+		if (Block.ShowInToolbarTopLevel.IsBound())
+		{
+			Visibility = UE::ToolMenus::Private::CombineVisibilityAttributes(
+				Visibility,
+				TAttribute<EVisibility>::CreateLambda([ShowInToolbarTopLevel = Block.ShowInToolbarTopLevel]
 				{
-					return EVisibility::Collapsed;
-				}
-
-				return EVisibility::Visible;
+					return ShowInToolbarTopLevel.Get() ? EVisibility::Visible : EVisibility::Collapsed;
+				})
+			);	
+		}
+		
+		if (Block.ScriptObject)
+		{
+			// Allow scripts to override ShowInToolbarTopLevel.
+			if (const TAttribute<bool> ScriptShowInToolbarTopLevel = Block.ScriptObject->CreateShowInToolbarTopLevelAttribute(MenuData->Context);
+				ScriptShowInToolbarTopLevel.IsSet())
+			{
+				Visibility = UE::ToolMenus::Private::CombineVisibilityAttributes(
+					TAttribute<EVisibility>::CreateLambda([ScriptShowInToolbarTopLevel]
+					{
+						return ScriptShowInToolbarTopLevel.Get() ? EVisibility::Visible : EVisibility::Collapsed;
+					}),
+					Visibility
+				);
 			}
-		);
+		}
 	}
+	
+	// Block level visibility applies last
+	Visibility = UE::ToolMenus::Private::CombineVisibilityAttributes(
+		Visibility,
+		Block.Visibility
+	);
 
 	const FUIAction UIAction = [&Block, MenuData]()
 		{
@@ -1899,7 +1939,7 @@ void UToolMenus::PopulateToolBarBuilderWithEntry(
 				Block.Icon,
 				Block.TutorialHighlightName,
 				FNewMenuDelegate(),
-				VisibilityOverride,
+				Visibility,
 				ToolbarLabelOverride,
 				ActualResizeParams
 			);
@@ -1914,32 +1954,6 @@ void UToolMenus::PopulateToolBarBuilderWithEntry(
 			UToolMenuEntryScript* ScriptObject = Block.ScriptObject;
 			const TAttribute<FSlateIcon> Icon = ScriptObject->CreateIconAttribute(MenuData->Context);
 
-			TAttribute<EVisibility> ScriptVisibilityOverride = VisibilityOverride;
-
-			// Allow scripts to override ShowInToolbarTopLevel.
-			if (const TAttribute<bool> ScriptShowInToolbarTopLevel =
-					ScriptObject->CreateShowInToolbarTopLevelAttribute(MenuData->Context);
-				bIsRaisingToTopLevel && ScriptShowInToolbarTopLevel.IsSet())
-			{
-				ScriptVisibilityOverride = TAttribute<EVisibility>::CreateLambda(
-					[OriginalVisibilityOverride = VisibilityOverride, ScriptShowInToolbarTopLevel]() -> EVisibility
-					{
-						if (ScriptShowInToolbarTopLevel.IsSet() && !ScriptShowInToolbarTopLevel.Get())
-						{
-							return EVisibility::Collapsed;
-						}
-						else if (OriginalVisibilityOverride.IsSet())
-						{
-							return OriginalVisibilityOverride.Get();
-						}
-						else
-						{
-							return EVisibility::Visible;
-						}
-					}
-				);
-			}
-
 			ToolBarBuilder.AddToolBarButton(
 				UIAction,
 				ScriptObject->Data.Name,
@@ -1948,7 +1962,7 @@ void UToolMenus::PopulateToolBarBuilderWithEntry(
 				Icon,
 				Block.UserInterfaceActionType,
 				Block.TutorialHighlightName,
-				ScriptVisibilityOverride,
+				Visibility,
 				ToolbarLabelOverride,
 				ActualResizeParams
 			);
@@ -1963,7 +1977,7 @@ void UToolMenus::PopulateToolBarBuilderWithEntry(
 				Block.Icon,
 				Block.UserInterfaceActionType,
 				Block.TutorialHighlightName,
-				VisibilityOverride,
+				Visibility,
 				ToolbarLabelOverride,
 				ActualResizeParams
 			);
@@ -1981,7 +1995,7 @@ void UToolMenus::PopulateToolBarBuilderWithEntry(
 				Block.Icon,
 				true,
 				Block.TutorialHighlightName,
-				VisibilityOverride,
+				Visibility,
 				ToolbarLabelOverride,
 				Block.ToolBarData.PlacementOverride,
 				Block.UserInterfaceActionType,
@@ -2023,7 +2037,7 @@ void UToolMenus::PopulateToolBarBuilderWithEntry(
 			Block.Icon,
 			Block.ToolBarData.bSimpleComboBox,
 			Block.TutorialHighlightName,
-			VisibilityOverride,
+			Visibility,
 			ToolbarLabelOverride,
 			Block.ToolBarData.PlacementOverride,
 			Block.UserInterfaceActionType,
@@ -2067,7 +2081,7 @@ void UToolMenus::PopulateToolBarBuilderWithEntry(
 			StyleParams.VerticalAlignment = VAlign_Fill;
 		}
 
-		ToolBarBuilder.AddWidget(Widget.ToSharedRef(), StyleParams, Block.TutorialHighlightName, Block.WidgetData.bSearchable, FNewMenuDelegate(), VisibilityOverride, ActualResizeParams);
+		ToolBarBuilder.AddWidget(Widget.ToSharedRef(), StyleParams, Block.TutorialHighlightName, Block.WidgetData.bSearchable, FNewMenuDelegate(), Visibility, ActualResizeParams);
 	}
 	else
 	{
@@ -2180,6 +2194,7 @@ void UToolMenus::PopulateToolBarBuilder(FToolBarBuilder& ToolBarBuilder, UToolMe
 				const bool bIsLastBlock = BlockIndex >= Section.Blocks.Num() - 1;
 				FToolMenuEntry& Block = Section.Blocks[BlockIndex];
 
+				// This is the top level of the toolbar; nothing is being elevated.
 				const bool bRaiseToTopLevel = false;
 				const bool bIsLastBlockOfLastSection = bIsLastBlock && bIsLastSection;
 				
@@ -2197,7 +2212,7 @@ void UToolMenus::PopulateToolBarBuilder(FToolBarBuilder& ToolBarBuilder, UToolMe
 				}
 				PreviousBlockGroupName = EntryBlockGroupName;
 				
-				PopulateToolBarBuilderWithEntry(ToolBarBuilder, MenuData, Block, bRaiseToTopLevel, bIsLastBlockOfLastSection);
+				PopulateToolBarBuilderWithEntry(ToolBarBuilder, MenuData, Section, Block, bRaiseToTopLevel, bIsLastBlockOfLastSection);
 			}
 			
 			if (!PreviousBlockGroupName.IsNone())
diff --git a/Engine/Source/Developer/ToolMenus/Public/ToolMenus.h b/Engine/Source/Developer/ToolMenus/Public/ToolMenus.h
--- a/Engine/Source/Developer/ToolMenus/Public/ToolMenus.h
+++ b/Engine/Source/Developer/ToolMenus/Public/ToolMenus.h
@@ -53,6 +53,19 @@ struct FGeneratedToolMenuWidgets
 	TArray<TSharedPtr<FGeneratedToolMenuWidget>> Instances;
 };
 
+namespace UE::ToolMenus
+{
+	/**
+	 * A handle to a particular entry. Used when elevating child elements to a parent toolbar.
+	 */
+	struct FSubBlockReference
+	{
+		UToolMenu* ParentMenu;
+		FToolMenuSection* Section;
+		FToolMenuEntry* Entry;
+	};	
+}
+
 /*
  * A global context that any menu can add/modify to specify which profiles are currently active
  */
@@ -394,15 +407,12 @@ private:
 	void PopulateMenuBuilder(FMenuBuilder& MenuBuilder, UToolMenu* MenuData);
 	void PopulateMenuBarBuilder(FMenuBarBuilder& MenuBarBuilder, UToolMenu* MenuData);
 
-	void ExtractChildBlocksFromSubMenu(
-		UToolMenu* ParentMenu, FToolMenuEntry& InBlock, TArray<TTuple<UToolMenu*, FToolMenuEntry*, FName>>& SubMenuBlocks
-	);
-	void PopulateToolBarBuilderWithTopLevelChildren(
-		FToolBarBuilder& ToolBarBuilder, UToolMenu* ParentMenu, FToolMenuEntry& InBlock, bool bAddSpaceAfterRaisedChildren
-	);
+	void ExtractChildBlocksFromSubMenu(UToolMenu* ParentMenu, FToolMenuEntry& InBlock, TArray<UE::ToolMenus::FSubBlockReference>& SubMenuBlocks);
+	void PopulateToolBarBuilderWithTopLevelChildren(FToolBarBuilder& ToolBarBuilder, UToolMenu* ParentMenu, FToolMenuEntry& InBlock, bool bAddSpaceAfterRaisedChildren);
 	void PopulateToolBarBuilderWithEntry(
 		FToolBarBuilder& ToolBarBuilder,
 		UToolMenu* MenuData,
+		FToolMenuSection& Section,
 		FToolMenuEntry& Entry,
 		bool bIsRaisingToTopLevel = false,
 		bool bIsLastBlockOfLastSection = false
diff --git a/Engine/Source/Editor/UnrealEd/Public/IPreviewLODController.h b/Engine/Source/Editor/UnrealEd/Public/IPreviewLODController.h
--- a/Engine/Source/Editor/UnrealEd/Public/IPreviewLODController.h
+++ b/Engine/Source/Editor/UnrealEd/Public/IPreviewLODController.h
@@ -33,5 +33,5 @@ public:
 	 * Set which LOD level an auto-generated list of LOD options should start from.
 	 * Allows certain LOD levels to be handled by commands provided by `FillLODCommands()`.
 	 */
-	virtual int32 GetAutoLODStartingIndex() const { return 0; }
+	virtual int32 GetAutoLODStartingIndex() const { return INDEX_NONE; }
 };
\ No newline at end of file
diff --git a/Engine/Source/Runtime/Experimental/Dataflow/Simulation/Public/Dataflow/DataflowSimulationManager.h b/Engine/Source/Runtime/Experimental/Dataflow/Simulation/Public/Dataflow/DataflowSimulationManager.h
--- a/Engine/Source/Runtime/Experimental/Dataflow/Simulation/Public/Dataflow/DataflowSimulationManager.h
+++ b/Engine/Source/Runtime/Experimental/Dataflow/Simulation/Public/Dataflow/DataflowSimulationManager.h
@@ -106,6 +106,9 @@ public:
 	/** Start the simulation tasks given a delta time */
 	void StartSimulationTasks(const float DeltaTime, const float SimulationTime);
 
+	/** Gets whether the simulation is enabled or disabled */
+	bool GetSimulationEnabled() const { return bIsSimulationEnabled; }
+
 	/** Set the simulation flag to enable/disable the simulation */
 	void SetSimulationEnabled(const bool bSimulationEnabled) {bIsSimulationEnabled = bSimulationEnabled;}
 
