Fix a compilation issue with UE::ToolMenus::FSubBlockReference on android.

diff --git a/Engine/Source/Developer/ToolMenus/Private/ToolMenus.cpp b/Engine/Source/Developer/ToolMenus/Private/ToolMenus.cpp
--- a/Engine/Source/Developer/ToolMenus/Private/ToolMenus.cpp
+++ b/Engine/Source/Developer/ToolMenus/Private/ToolMenus.cpp
@@ -70,6 +70,33 @@ TAttribute<EVisibility> CombineVisibilityAttributes(const TAttribute<EVisibility
 
 } // namespace UE::ToolMenus::Private
 
+namespace UE::ToolMenus
+{
+
+FSubBlockReference::FSubBlockReference()
+	: ParentMenu(nullptr)
+	, Section(nullptr)
+	, Entry(nullptr)
+{
+}
+
+FSubBlockReference::FSubBlockReference(const FSubBlockReference& Other)
+	: ParentMenu(Other.ParentMenu)
+	, Section(Other.Section)
+	, Entry(Other.Entry)
+{
+}
+
+FSubBlockReference::FSubBlockReference(UToolMenu* InParent, FToolMenuSection* InSection, FToolMenuEntry* InEntry)
+	: ParentMenu(InParent)
+	, Section(InSection)
+	, Entry(InEntry)
+{
+}
+
+
+} // namespace UE::ToolMenus
+
 UToolMenus* UToolMenus::Singleton = nullptr;
 bool UToolMenus::bHasShutDown = false;
 FSimpleMulticastDelegate UToolMenus::StartupCallbacks;
diff --git a/Engine/Source/Developer/ToolMenus/Public/ToolMenus.h b/Engine/Source/Developer/ToolMenus/Public/ToolMenus.h
--- a/Engine/Source/Developer/ToolMenus/Public/ToolMenus.h
+++ b/Engine/Source/Developer/ToolMenus/Public/ToolMenus.h
@@ -60,6 +60,10 @@ namespace UE::ToolMenus
 	 */
 	struct FSubBlockReference
 	{
+		FSubBlockReference();
+		FSubBlockReference(const FSubBlockReference& Other);
+		FSubBlockReference(UToolMenu* InParent, FToolMenuSection* InSection, FToolMenuEntry* InEntry);
+
 		UToolMenu* ParentMenu;
 		FToolMenuSection* Section;
 		FToolMenuEntry* Entry;
