Fix a spelling error in ToolMenuDelegates.h

diff --git a/Engine/Source/Developer/ToolMenus/Public/ToolMenuDelegates.h b/Engine/Source/Developer/ToolMenus/Public/ToolMenuDelegates.h
--- a/Engine/Source/Developer/ToolMenus/Public/ToolMenuDelegates.h
+++ b/Engine/Source/Developer/ToolMenus/Public/ToolMenuDelegates.h
@@ -48,8 +48,8 @@ struct TOOLMENUS_API FToolUIAction
 		, GetActionCheckState(InGetActionCheckState)
 	{}
 	
-	FToolUIAction(const FToolMenuExecuteAction& InExectueAction, const FToolMenuGetActionCheckState& InGetActionCheckState)
-		: ExecuteAction(InExectueAction)
+	FToolUIAction(const FToolMenuExecuteAction& InExecuteAction, const FToolMenuGetActionCheckState& InGetActionCheckState)
+		: ExecuteAction(InExecuteAction)
 		, GetActionCheckState(InGetActionCheckState)
 	{}
 
