Fix an issue where the level editor viewport toolbar had extra space on the right.
- Place legacy extension sections before the viewport sizing elements.
- Make whether or not adjacent toolbar sections place separators between them optional for ToolMenu-derived toolbars.
- Fix a bug with toolbar section ordering logic where the "last section" would not be correctly identified.

diff --git a/Engine/Source/Developer/ToolMenus/Private/ToolMenu.cpp b/Engine/Source/Developer/ToolMenus/Private/ToolMenu.cpp
--- a/Engine/Source/Developer/ToolMenus/Private/ToolMenu.cpp
+++ b/Engine/Source/Developer/ToolMenus/Private/ToolMenu.cpp
@@ -92,6 +92,7 @@ void UToolMenu::InitGeneratedCopy(const UToolMenu* Source, const FName InMenuNam
 	bShouldCloseWindowAfterMenuSelection = Source->bShouldCloseWindowAfterMenuSelection;
 	bCloseSelfOnly = Source->bCloseSelfOnly;
 	bSearchable = Source->bSearchable;
+	bSeparateSections = Source->bSeparateSections;
 	bToolBarIsFocusable = Source->bToolBarIsFocusable;
 	bToolBarForceSmallIcons = Source->bToolBarForceSmallIcons;
 	MenuOwner = Source->MenuOwner;
diff --git a/Engine/Source/Developer/ToolMenus/Private/ToolMenus.cpp b/Engine/Source/Developer/ToolMenus/Private/ToolMenus.cpp
--- a/Engine/Source/Developer/ToolMenus/Private/ToolMenus.cpp
+++ b/Engine/Source/Developer/ToolMenus/Private/ToolMenus.cpp
@@ -68,6 +68,24 @@ TAttribute<EVisibility> CombineVisibilityAttributes(const TAttribute<EVisibility
 	return TAttribute<EVisibility>();
 }
 
+int32 SortSectionAlignment(EToolMenuSectionAlign A, EToolMenuSectionAlign B)
+{
+	if (A == B)
+	{
+		return 0;
+	}
+	
+	if (A == EToolMenuSectionAlign::First && B == EToolMenuSectionAlign::Default)
+	{
+		return -1;
+	}
+	if (A == EToolMenuSectionAlign::Default && B == EToolMenuSectionAlign::First)
+	{
+		return 1;
+	}
+	return A > B ? 1 : -1;
+}
+
 } // namespace UE::ToolMenus::Private
 
 namespace UE::ToolMenus
@@ -1920,7 +1938,7 @@ void UToolMenus::PopulateToolBarBuilderWithEntry(
 				return UToolMenus::ConvertUIAction(Block, MenuData->Context);
 			}
 		}();
-
+	
 	TAttribute<FText> ToolbarLabelOverride;
 	if (Block.ToolBarData.LabelOverride.IsSet())
 	{
@@ -2147,112 +2165,105 @@ void UToolMenus::PopulateToolBarBuilder(FToolBarBuilder& ToolBarBuilder, UToolMe
 	// Add the sections grouped by alignment with SSpacers in between. This visually separates them and allows users
 	// to align sections to appear first, middle, or last. Default-aligned sections appear grouped with first-aligned
 	// sections but appear after them.
-	const TArray<EToolMenuSectionAlign> SectionAlignments = { EToolMenuSectionAlign::First,
-															  EToolMenuSectionAlign::Default,
-															  EToolMenuSectionAlign::Middle,
-															  EToolMenuSectionAlign::Last };
-	int32 NumAddedSections = 0;
-	for (const EToolMenuSectionAlign CurrentAlignment : SectionAlignments)
+	TArray<FToolMenuSection*> SortedSections;
+	SortedSections.Reserve(MenuData->Sections.Num());
+	for (int32 Index = 0; Index < MenuData->Sections.Num(); ++Index)
+	{
+		SortedSections.Add(&MenuData->Sections[Index]);
+	}
+	
+	SortedSections.StableSort([](const FToolMenuSection& A, const FToolMenuSection& B)
+	{
+		return UE::ToolMenus::Private::SortSectionAlignment(A.Alignment, B.Alignment) < 0;
+	});
+	
+	EToolMenuSectionAlign LastAlignment = EToolMenuSectionAlign::First;
+	for (int32 SectionIndex = 0; SectionIndex < SortedSections.Num(); ++SectionIndex)
 	{
-		// If we've already added all sections, there's no more work to do. Especially, we don't want to add any more SSpacers.
-		if (NumAddedSections >= MenuData->Sections.Num())
+		FToolMenuSection& Section = *SortedSections[SectionIndex];
+		
+		if (Section.Alignment != LastAlignment)
 		{
-			break;
-		}
+			const bool bIsMiddleOrLast = Section.Alignment == EToolMenuSectionAlign::Middle || Section.Alignment == EToolMenuSectionAlign::Last;
 
-		const bool bIsMiddleOrLast = CurrentAlignment == EToolMenuSectionAlign::Middle
-								  || CurrentAlignment == EToolMenuSectionAlign::Last;
+			// Add a spacer before the middle and last alignment groups, and only if we've already added a section
+			if (bIsMiddleOrLast && SectionIndex > 0)
+			{
+				static constexpr float AlmostZero = 0.1f; // Using 0.0f results in different behavior, so we just use a small, sub-1 pixel value.
 
-		// Add a spacer before the middle and last alignment groups, and only if we've already added a section to a
-		// previous alignment group.
-		if (bIsMiddleOrLast && NumAddedSections > 0)
-		{
-			static constexpr float AlmostZero = 0.1f; // Using 0.0f results in different behavior, so we just use a small, sub-1 pixel value.
+				FMenuEntryStyleParams StyleParams;
 
-			FMenuEntryStyleParams StyleParams;
+				StyleParams.HorizontalAlignment = HAlign_Right;
+				StyleParams.SizeRule = FSizeParam::ESizeRule::SizeRule_StretchContent;
+				StyleParams.FillSize = 1.0f;
+				StyleParams.FillSizeMin = AlmostZero; // Allow the spacer to shrink to nothing if another widget needs the space.
+				StyleParams.MinimumSize = AlmostZero;
+				StyleParams.DesiredWidthOverride = StyleParams.DesiredHeightOverride = AlmostZero;
 
-			StyleParams.HorizontalAlignment = HAlign_Right;
-			StyleParams.SizeRule = FSizeParam::ESizeRule::SizeRule_StretchContent;
-			StyleParams.FillSize = 1.0f;
-			StyleParams.FillSizeMin = AlmostZero; // Allow the spacer to shrink to nothing if another widget needs the space.
-			StyleParams.MinimumSize = AlmostZero;
-			StyleParams.DesiredWidthOverride = StyleParams.DesiredHeightOverride = AlmostZero;
+				FMenuEntryResizeParams ResizeParams;
+				// Never show spacers in overflow menus.
+				ResizeParams.VisibleInOverflow = false;
+				// Prevent our spacers from overflowing. This will keep them in the toolbar, but they will still shrink so
+				// overflow of other entries will not be affected. What will be affected, however, is stretching of the
+				// toolbar after entries have been clipped. By keeping the stretchers in the toolbar, we can grow them to
+				// fill up any space left over when a widget is clipped.
+				ResizeParams.AllowClipping = false;
 
-			FMenuEntryResizeParams ResizeParams;
-			// Never show spacers in overflow menus.
-			ResizeParams.VisibleInOverflow = false;
-			// Prevent our spacers from overflowing. This will keep them in the toolbar, but they will still shrink so
-			// overflow of other entries will not be affected. What will be affected, however, is stretching of the
-			// toolbar after entries have been clipped. By keeping the stretchers in the toolbar, we can grow them to
-			// fill up any space left over when a widget is clipped.
-			ResizeParams.AllowClipping = false;
-
-			ToolBarBuilder.AddWidget(
-				SNew(SSpacer), StyleParams, NAME_None, false, FNewMenuDelegate(), TAttribute<EVisibility>(), ResizeParams
-			);
+				ToolBarBuilder.AddWidget(
+					SNew(SSpacer), StyleParams, NAME_None, false, FNewMenuDelegate(), TAttribute<EVisibility>(), ResizeParams
+				);
+			}
 		}
-
-		// Keep track if this is the first section we're adding for the CurrentAlignment. Make an exception if the
-		// current alignment is Default and we already added a first-aligned section, because in that case this
-		// isn't the first section in the "group" since we're displaying first and default-aligned sections together.
-		bool bFirstSectionInAlignmentGroup = CurrentAlignment == EToolMenuSectionAlign::Default ? NumAddedSections == 0
-																								: true;
-		for (int SectionIndex = 0; SectionIndex < MenuData->Sections.Num(); ++SectionIndex)
+		
+		if (Section.Construct.NewToolBarDelegateLegacy.IsBound())
 		{
-			const bool bIsLastSection = SectionIndex >= MenuData->Sections.Num() - 1;
-			FToolMenuSection& Section = MenuData->Sections[SectionIndex];
-
-			if (Section.Alignment != CurrentAlignment)
-			{
-				continue;
-			}
-
-			if (Section.Construct.NewToolBarDelegateLegacy.IsBound())
-			{
-				Section.Construct.NewToolBarDelegateLegacy.Execute(ToolBarBuilder, MenuData);
-				continue;
-			}
+			Section.Construct.NewToolBarDelegateLegacy.Execute(ToolBarBuilder, MenuData);
+			continue;
+		}
+		
+		const bool bIsLastSection = SectionIndex >= SortedSections.Num() - 1;
+		// "Default" sections are placed directly after "first" sections, but on the same side.
+		// Therefor, they should not be considered "first".
+		const bool bFirstSectionInAlignmentGroup = Section.Alignment != LastAlignment && Section.Alignment != EToolMenuSectionAlign::Default;
 
-			const bool bSectionShouldHaveSeparator = !bFirstSectionInAlignmentGroup;
-			ToolBarBuilder.BeginSection(Section.Name, bSectionShouldHaveSeparator);
+		const bool bSectionShouldHaveSeparator = MenuData->bSeparateSections && !bFirstSectionInAlignmentGroup;
+		ToolBarBuilder.BeginSection(Section.Name, bSectionShouldHaveSeparator);
 
-			FName PreviousBlockGroupName = NAME_None;
-			for (int BlockIndex = 0; BlockIndex < Section.Blocks.Num(); ++BlockIndex)
-			{
-				const bool bIsLastBlock = BlockIndex >= Section.Blocks.Num() - 1;
-				FToolMenuEntry& Block = Section.Blocks[BlockIndex];
+		FName PreviousBlockGroupName = NAME_None;
+		for (int BlockIndex = 0; BlockIndex < Section.Blocks.Num(); ++BlockIndex)
+		{
+			const bool bIsLastBlock = BlockIndex >= Section.Blocks.Num() - 1;
+			FToolMenuEntry& Block = Section.Blocks[BlockIndex];
 
-				// This is the top level of the toolbar; nothing is being elevated.
-				const bool bRaiseToTopLevel = false;
-				const bool bIsLastBlockOfLastSection = bIsLastBlock && bIsLastSection;
+			// This is the top level of the toolbar; nothing is being elevated.
+			const bool bRaiseToTopLevel = false;
+			const bool bIsLastBlockOfLastSection = bIsLastBlock && bIsLastSection;
 				
-				const FName EntryBlockGroupName = Block.ToolBarData.BlockGroupName;
-				if (EntryBlockGroupName != PreviousBlockGroupName)
+			const FName EntryBlockGroupName = Block.ToolBarData.BlockGroupName;
+			if (EntryBlockGroupName != PreviousBlockGroupName)
+			{
+				if (!PreviousBlockGroupName.IsNone())
 				{
-					if (!PreviousBlockGroupName.IsNone())
-					{
-						ToolBarBuilder.EndBlockGroup();
-					}
-					if (!EntryBlockGroupName.IsNone())
-					{
-						ToolBarBuilder.BeginBlockGroup();
-					}
+					ToolBarBuilder.EndBlockGroup();
+				}
+				if (!EntryBlockGroupName.IsNone())
+				{
+					ToolBarBuilder.BeginBlockGroup();
 				}
-				PreviousBlockGroupName = EntryBlockGroupName;
-				
-				PopulateToolBarBuilderWithEntry(ToolBarBuilder, MenuData, Section, Block, bRaiseToTopLevel, bIsLastBlockOfLastSection);
 			}
+			PreviousBlockGroupName = EntryBlockGroupName;
+				
+			PopulateToolBarBuilderWithEntry(ToolBarBuilder, MenuData, Section, Block, bRaiseToTopLevel, bIsLastBlockOfLastSection);
+		}
 			
-			if (!PreviousBlockGroupName.IsNone())
-			{
-				ToolBarBuilder.EndBlockGroup();
-			}
-
-			ToolBarBuilder.EndSection();
-
-			++NumAddedSections;
-			bFirstSectionInAlignmentGroup = false;
+		if (!PreviousBlockGroupName.IsNone())
+		{
+			ToolBarBuilder.EndBlockGroup();
 		}
+
+		ToolBarBuilder.EndSection();
+		
+		LastAlignment = Section.Alignment;
 	}
 
 	AddReferencedContextObjects(ToolBarBuilder.GetMultiBox(), MenuData);
diff --git a/Engine/Source/Developer/ToolMenus/Public/ToolMenu.h b/Engine/Source/Developer/ToolMenus/Public/ToolMenu.h
--- a/Engine/Source/Developer/ToolMenus/Public/ToolMenu.h
+++ b/Engine/Source/Developer/ToolMenus/Public/ToolMenu.h
@@ -195,6 +195,10 @@ public:
 
 	UPROPERTY(EditAnywhere, BlueprintReadWrite, Category = "ToolBar")
 	bool bToolBarIsFocusable;
+	
+	/** Whether different sections should have separators automatically added between them. */
+	UPROPERTY(EditAnywhere, BlueprintReadWrite, Category = "ToolBar")
+	bool bSeparateSections = true;
 
 	/** Whether to (ever) allow the wrap button/overflow menu. Defauls to the style. */
 	UPROPERTY(EditAnywhere, BlueprintReadWrite, Category = "ToolBar")
diff --git a/Engine/Source/Editor/LevelEditor/Private/SLevelViewport.cpp b/Engine/Source/Editor/LevelEditor/Private/SLevelViewport.cpp
--- a/Engine/Source/Editor/LevelEditor/Private/SLevelViewport.cpp
+++ b/Engine/Source/Editor/LevelEditor/Private/SLevelViewport.cpp
@@ -2048,6 +2048,7 @@ TSharedPtr<SWidget> SLevelViewport::BuildViewportToolbar()
 			);
 
 			ViewportToolbarMenu->StyleName = "LevelEditorViewportToolbar";
+			ViewportToolbarMenu->bSeparateSections = false;
 
 			// Add the left-aligned part of the viewport toolbar.
 			{
@@ -2139,16 +2140,14 @@ TSharedPtr<SWidget> SLevelViewport::BuildViewportToolbar()
 
 				RightSection.AddEntry(UE::LevelEditor::CreatePerformanceAndScalabilitySubmenu());
 				RightSection.AddEntry(UE::LevelEditor::CreateSettingsSubmenu());
-				RightSection.AddEntry(UE::LevelEditor::CreateViewportSizingSubmenu());
 			}
 
 			// Sections added for compatibility with old toolbar transform menu extensions
 			{
 				// Existing code might add entries via FToolBarBuilder Extenders,
-				// which seem to only work with Toolbar main sections.
-				// The following sections match the ones added in STransformViewportToolBar::MakeTransformToolBar(...),
-				// which makes use of a FToolBarBuilder. Like so, code extending the Old Transform Toolbar using these
-				// sections names as hooks should also work in the new Level Viewport Toolbar.
+				// These extenders are based off of toolbar section & extension names.
+				// The following sections names are used in STransformViewportToolBar::MakeTransformToolBar(...),
+				// This allows legacy extensions to be added directly to the toolbar.
 
 				ViewportToolbarMenu->FindOrAddSection("Transform").Alignment = EToolMenuSectionAlign::Last;
 				ViewportToolbarMenu->FindOrAddSection("LocationGridSnap").Alignment = EToolMenuSectionAlign::Last;
@@ -2157,6 +2156,14 @@ TSharedPtr<SWidget> SLevelViewport::BuildViewportToolbar()
 				ViewportToolbarMenu->FindOrAddSection("ScaleGridSnap").Alignment = EToolMenuSectionAlign::Last;
 				ViewportToolbarMenu->FindOrAddSection("CameraSpeed").Alignment = EToolMenuSectionAlign::Last;
 			}
+			
+			{
+				// The viewport sizing entries should _always_ be the last section.
+				// Doing otherwise breaks the customized behavior of the overflow menu.
+				FToolMenuSection& SizingSection = ViewportToolbarMenu->FindOrAddSection("ViewportSizing");
+				SizingSection.Alignment = EToolMenuSectionAlign::Last;
+				SizingSection.AddEntry(UE::LevelEditor::CreateViewportSizingSubmenu());
+			}
 		}
 	}
 
