Fix an issue where raised menu items could be invisible while still making extra space visible.
- Refactor UToolMenu visibility to FToolMenuVisibilityChoice instead of TAttribute<EVisibility>, improving the ergonomics of working with visibility in ToolMenus.

diff --git a/Engine/Plugins/Experimental/Dataflow/Source/DataflowEditor/Private/Dataflow/DataflowSimulationViewport.cpp b/Engine/Plugins/Experimental/Dataflow/Source/DataflowEditor/Private/Dataflow/DataflowSimulationViewport.cpp
--- a/Engine/Plugins/Experimental/Dataflow/Source/DataflowEditor/Private/Dataflow/DataflowSimulationViewport.cpp
+++ b/Engine/Plugins/Experimental/Dataflow/Source/DataflowEditor/Private/Dataflow/DataflowSimulationViewport.cpp
@@ -189,7 +189,7 @@ TSharedPtr<SWidget> SDataflowSimulationViewport::BuildViewportToolbar()
 					Pause.ToolBarData = ToolBarData;
 					Pause.Icon = FSlateIcon(FAppStyle::Get().GetStyleSetName(), "Animation.Pause");
 					Pause.SetShowInToolbarTopLevel(true);
-					Pause.SetVisibility(TAttribute<EVisibility>::CreateLambda([WeakViewport]
+					Pause.Visibility = [WeakViewport]
 					{
 						if (TSharedPtr<SDataflowSimulationViewport> Viewport = WeakViewport.Pin())
 						{
@@ -197,18 +197,18 @@ TSharedPtr<SWidget> SDataflowSimulationViewport::BuildViewportToolbar()
 							{
 								if (FDataflowSimulationScene* SimulationScene = StaticCast<FDataflowSimulationScene*>(Client->GetPreviewScene()))
 								{
-									return SimulationScene->IsSimulationEnabled() ? EVisibility::Visible : EVisibility::Collapsed;
+									return SimulationScene->IsSimulationEnabled();
 								}
 							}
 						}
-						return EVisibility::Visible;
-					}));
+						return true;
+					};
 					
 					FToolMenuEntry& Start = Section.AddMenuEntry(FDataflowEditorCommands::Get().StartSimulationScene);
 					Start.ToolBarData = ToolBarData;
 					Start.Icon = FSlateIcon(FAppStyle::Get().GetStyleSetName(), "Icons.Play");
 					Start.SetShowInToolbarTopLevel(true);
-					Start.SetVisibility(TAttribute<EVisibility>::CreateLambda([WeakViewport]
+					Start.Visibility = [WeakViewport]
                     {
                     	if (TSharedPtr<SDataflowSimulationViewport> Viewport = WeakViewport.Pin())
                     	{
@@ -216,12 +216,12 @@ TSharedPtr<SWidget> SDataflowSimulationViewport::BuildViewportToolbar()
                     		{
                     			if (FDataflowSimulationScene* SimulationScene = StaticCast<FDataflowSimulationScene*>(Client->GetPreviewScene()))
                     			{
-                    				return SimulationScene->IsSimulationEnabled() ? EVisibility::Collapsed : EVisibility::Visible;
+                    				return !SimulationScene->IsSimulationEnabled();
                     			}
                     		}
                     	}
-                    	return EVisibility::Visible;
-                    }));
+                    	return true;
+                    };
 					
 					FToolMenuEntry& Step = Section.AddMenuEntry(FDataflowEditorCommands::Get().StepSimulationScene);
 					Step.ToolBarData = ToolBarData;
diff --git a/Engine/Source/Developer/ToolMenus/Private/ToolMenuEntry.cpp b/Engine/Source/Developer/ToolMenus/Private/ToolMenuEntry.cpp
--- a/Engine/Source/Developer/ToolMenus/Private/ToolMenuEntry.cpp
+++ b/Engine/Source/Developer/ToolMenus/Private/ToolMenuEntry.cpp
@@ -13,6 +13,117 @@
 
 #include UE_INLINE_GENERATED_CPP_BY_NAME(ToolMenuEntry)
 
+FToolMenuVisibilityChoice::FToolMenuVisibilityChoice()
+{
+}
+
+FToolMenuVisibilityChoice::FToolMenuVisibilityChoice(const FToolMenuVisibilityChoice& Other)
+	: Value(Other.Value)
+{
+}
+
+FToolMenuVisibilityChoice::FToolMenuVisibilityChoice(const TAttribute<EVisibility>& InVisibility)
+{
+	Value.Set<TAttribute<EVisibility>>(InVisibility);
+}
+
+FToolMenuVisibilityChoice::FToolMenuVisibilityChoice(const TAttribute<bool>& InIsVisible)
+{
+	Value.Set<TAttribute<bool>>(InIsVisible);
+}
+
+FToolMenuVisibilityChoice::FToolMenuVisibilityChoice(const FIsActionButtonVisible& InIsActionButtonVisible)
+{
+	Value.Set<FIsActionButtonVisible>(InIsActionButtonVisible);
+}
+
+FToolMenuVisibilityChoice& FToolMenuVisibilityChoice::operator=(TFunction<EVisibility()> VisibilityFunc)
+{
+	Value.Set<TAttribute<EVisibility>>(TAttribute<EVisibility>::CreateLambda(VisibilityFunc));
+	return *this;
+}
+
+FToolMenuVisibilityChoice& FToolMenuVisibilityChoice::operator=(TFunction<bool()> IsVisibleFunc)
+{
+	Value.Set<TAttribute<bool>>(TAttribute<bool>::CreateLambda(IsVisibleFunc));
+	return *this;
+}
+
+TAttribute<EVisibility> FToolMenuVisibilityChoice::ToVisibilityAttribute() const
+{
+	if (const TAttribute<EVisibility>* Visibility = Value.TryGet<TAttribute<EVisibility>>())
+	{
+		return *Visibility;
+	}
+	
+	if (const TAttribute<bool>* IsVisible = Value.TryGet<TAttribute<bool>>())
+	{
+		return TAttribute<EVisibility>::CreateLambda([IsVisible = *IsVisible]
+		{
+			return IsVisible.Get() ? EVisibility::Visible : EVisibility::Collapsed;
+		});
+	}
+	
+	if (const FIsActionButtonVisible* Action = Value.TryGet<FIsActionButtonVisible>())
+	{
+		if (Action->IsBound())
+		{
+			return TAttribute<EVisibility>::CreateLambda([Action = *Action]
+			{
+				return Action.Execute() ? EVisibility::Visible : EVisibility::Collapsed;
+			});
+		}
+	}
+	
+	return TAttribute<EVisibility>();
+}
+
+bool FToolMenuVisibilityChoice::IsSet() const
+{
+	if (const TAttribute<EVisibility>* Visibility = Value.TryGet<TAttribute<EVisibility>>())
+	{
+		return Visibility->IsSet();
+	}
+	
+	if (const TAttribute<bool>* IsVisible = Value.TryGet<TAttribute<bool>>())
+	{
+		return IsVisible->IsSet();
+	}
+	
+	if (const FIsActionButtonVisible* Action = Value.TryGet<FIsActionButtonVisible>())
+	{
+		return Action->IsBound();
+	}
+	
+	return false;
+}
+
+EVisibility FToolMenuVisibilityChoice::Get() const
+{
+	if (const TAttribute<EVisibility>* Visibility = Value.TryGet<TAttribute<EVisibility>>())
+	{
+		return Visibility->Get();
+	}
+	
+	if (const TAttribute<bool>* IsVisible = Value.TryGet<TAttribute<bool>>())
+	{
+		if (IsVisible->IsSet())
+		{
+			return IsVisible->Get() ? EVisibility::Visible : EVisibility::Collapsed;
+		}
+	}
+	
+	if (const FIsActionButtonVisible* Action = Value.TryGet<FIsActionButtonVisible>())
+	{
+		if (Action->IsBound())
+		{
+			return Action->Execute() ? EVisibility::Visible : EVisibility::Collapsed;
+		}
+	}
+	
+	return EVisibility::Visible;
+}
+
 PRAGMA_DISABLE_DEPRECATION_WARNINGS
 FToolMenuEntry::FToolMenuEntry()
 	: Type(EMultiBlockType::None)
@@ -339,11 +450,6 @@ void FToolMenuEntry::SetShowInToolbarTopLevel(TAttribute<bool> InTopLevel)
 	ShowInToolbarTopLevel = InTopLevel;
 }
 
-void FToolMenuEntry::SetVisibility(TAttribute<EVisibility> InVisibility)
-{
-	Visibility = InVisibility;
-}
-
 bool FToolMenuEntry::IsScriptObjectDynamicConstruct() const
 {
 	static const FName ConstructMenuEntryName = GET_FUNCTION_NAME_CHECKED(UToolMenuEntryScript, ConstructMenuEntry);
diff --git a/Engine/Source/Developer/ToolMenus/Private/ToolMenus.cpp b/Engine/Source/Developer/ToolMenus/Private/ToolMenus.cpp
--- a/Engine/Source/Developer/ToolMenus/Private/ToolMenus.cpp
+++ b/Engine/Source/Developer/ToolMenus/Private/ToolMenus.cpp
@@ -35,21 +35,18 @@ UToolMenus* CreateToolMenusInstance()
 }
 
 /** Combine two visibility attributes with the first attribute's non-visibility taking precedence. */
-TAttribute<EVisibility> CombineVisibilityAttributes(const TAttribute<EVisibility>& PrimaryVisibility, const TAttribute<EVisibility>& SecondaryVisibility)
+FToolMenuVisibilityChoice CombineVisibility(const FToolMenuVisibilityChoice& PrimaryVisibility, const FToolMenuVisibilityChoice& SecondaryVisibility)
 {
 	if (PrimaryVisibility.IsSet() && SecondaryVisibility.IsSet())
 	{
 		return TAttribute<EVisibility>::CreateLambda(
 			[Primary = PrimaryVisibility, Secondary = SecondaryVisibility]()
 			{
-				// If section visibility is set, and section is not visible, return Section Visibility
-				const EVisibility PrimaryVisibility = Primary.Get();
-				if (PrimaryVisibility != EVisibility::Visible)
+				const EVisibility PrimaryValue = Primary.Get(); 
+				if (PrimaryValue != EVisibility::Visible)
 				{
-					return PrimaryVisibility;
+					return PrimaryValue;
 				}
-
-				// Stick to this Block visibility
 				return Secondary.Get();
 			}
 		);
@@ -65,7 +62,7 @@ TAttribute<EVisibility> CombineVisibilityAttributes(const TAttribute<EVisibility
 		return SecondaryVisibility;
 	}
 		
-	return TAttribute<EVisibility>();
+	return FToolMenuVisibilityChoice();
 }
 
 int32 SortSectionAlignment(EToolMenuSectionAlign A, EToolMenuSectionAlign B)
@@ -256,7 +253,7 @@ public:
 
 		if (!bSubMenuAdded)
 		{
-			const TAttribute<EVisibility> VisibilityOverride = UE::ToolMenus::Private::CombineVisibilityAttributes(Section.Visibility, Block.Visibility);
+			const FToolMenuVisibilityChoice VisibilityOverride = UE::ToolMenus::Private::CombineVisibility(Section.Visibility, Block.Visibility);
 
 			if (Widget.IsValid())
 			{
@@ -315,7 +312,7 @@ public:
 
 	void AddStandardEntryToMenuBuilder()
 	{
-		const TAttribute<EVisibility> VisibilityOverride = UE::ToolMenus::Private::CombineVisibilityAttributes(Section.Visibility, Block.Visibility);
+		const FToolMenuVisibilityChoice VisibilityOverride = UE::ToolMenus::Private::CombineVisibility(Section.Visibility, Block.Visibility);
 
 		// First, check for a ToolUIAction, otherwise do the rest of this (have CommandList and Command)
 		// Need another variable to store if we are using a keybind from a command
@@ -408,7 +405,7 @@ public:
 			return;
 		}
 
-		const TAttribute<EVisibility> VisibilityOverride = UE::ToolMenus::Private::CombineVisibilityAttributes(Section.Visibility, Block.Visibility);
+		const FToolMenuVisibilityChoice VisibilityOverride = UE::ToolMenus::Private::CombineVisibility(Section.Visibility, Block.Visibility);
 
 		UIAction = UToolMenus::ConvertUIAction(Block, MenuData->Context);
 		bUIActionIsSet = UIAction.ExecuteAction.IsBound() || UIAction.CanExecuteAction.IsBound() || UIAction.GetActionCheckState.IsBound() || UIAction.IsActionVisibleDelegate.IsBound();
@@ -1669,7 +1666,7 @@ void UToolMenus::PopulateToolBarBuilderWithTopLevelChildren(FToolBarBuilder& Too
 	}
 
 	// Dynamic visibility of trailing raised-children spacer.
-	TArray<TAttribute<bool>> AllRaisedVisibilities = TArray<TAttribute<bool>>();
+	TArray<TAttribute<EVisibility>> AllRaisedVisibilities = TArray<TAttribute<EVisibility>>();
 
 	// Dynamic visibility of separators
 	//
@@ -1692,9 +1689,9 @@ void UToolMenus::PopulateToolBarBuilderWithTopLevelChildren(FToolBarBuilder& Too
 	//   sep_vis = anyVisible(nextEntries) && anyVisible(PreviousEntries
 	//
 
-	TArray<TAttribute<bool>> PreviousVisibilities;
+	TArray<TAttribute<EVisibility>> PreviousVisibilities;
 	// This has to be heap allocated so we can still add do it after a separator visibility delegate captures it.
-	TSharedPtr<TArray<TAttribute<bool>>> NextVisibilities = MakeShared<TArray<TAttribute<bool>>>();
+	TSharedPtr<TArray<TAttribute<EVisibility>>> NextVisibilities = MakeShared<TArray<TAttribute<EVisibility>>>();
 
 	bool bHasRaisedEntrySinceLastSeparator = false;
 	// Seed the previous section with the first block's section so we don't start with adding a separator because
@@ -1713,7 +1710,7 @@ void UToolMenus::PopulateToolBarBuilderWithTopLevelChildren(FToolBarBuilder& Too
 		{
 			// Step entry visibility delegate records forward now that we encountered a new separator.
 			PreviousVisibilities.Append(*NextVisibilities);
-			NextVisibilities = MakeShared<TArray<TAttribute<bool>>>();
+			NextVisibilities = MakeShared<TArray<TAttribute<EVisibility>>>();
 
 			const TAttribute<EVisibility> VisibilityOverride = TAttribute<EVisibility>::CreateLambda(
 				[Previous = PreviousVisibilities, Next = NextVisibilities]()
@@ -1722,9 +1719,9 @@ void UToolMenus::PopulateToolBarBuilderWithTopLevelChildren(FToolBarBuilder& Too
 					// const bool bVisible = bAnyNext && bAnyPrevious;
 
 					bool bAnyNext = false;
-					for (const TAttribute<bool>& Visibility : *Next)
+					for (const TAttribute<EVisibility>& Visibility : *Next)
 					{
-						if (Visibility.Get())
+						if (Visibility.Get() == EVisibility::Visible)
 						{
 							bAnyNext = true;
 							break;
@@ -1736,9 +1733,9 @@ void UToolMenus::PopulateToolBarBuilderWithTopLevelChildren(FToolBarBuilder& Too
 						return EVisibility::Collapsed;
 					}
 
-					for (const TAttribute<bool>& Visibility : Previous)
+					for (const TAttribute<EVisibility>& Visibility : Previous)
 					{
-						if (Visibility.Get())
+						if (Visibility.Get() == EVisibility::Visible)
 						{
 							return EVisibility::Visible;
 						}
@@ -1777,13 +1774,16 @@ void UToolMenus::PopulateToolBarBuilderWithTopLevelChildren(FToolBarBuilder& Too
 			constexpr bool bRaiseToTopLevel = true;
 			PopulateToolBarBuilderWithEntry(ToolBarBuilder, SubMenu, *Section, *Entry, bRaiseToTopLevel);
 			bHasRaisedEntrySinceLastSeparator = true;
+			
+			constexpr bool bEmbedActionOrCommand = true;
+			const TAttribute<EVisibility> EntryVisibility = CalculateToolbarVisibility(SubMenu, *Section, *Entry, bRaiseToTopLevel, bEmbedActionOrCommand);
 
 			// Keep track of added entries' visibilities so separators can set their visibility override.
-			NextVisibilities->Add(Entry->ShowInToolbarTopLevel);
+			NextVisibilities->Add(EntryVisibility);
 
 			if (bAddSpaceAfterRaisedChildren)
 			{
-				AllRaisedVisibilities.Add(Entry->ShowInToolbarTopLevel);
+				AllRaisedVisibilities.Add(EntryVisibility);
 			}
 		}
 
@@ -1805,9 +1805,9 @@ void UToolMenus::PopulateToolBarBuilderWithTopLevelChildren(FToolBarBuilder& Too
 		TAttribute<EVisibility> SpacerVisibilityOverride = TAttribute<EVisibility>::CreateLambda(
 			[AllRaisedVisibilities]()
 			{
-				for (const TAttribute<bool>& Visibility : AllRaisedVisibilities)
+				for (const TAttribute<EVisibility>& Visibility : AllRaisedVisibilities)
 				{
-					if (Visibility.Get())
+					if (Visibility.Get() == EVisibility::Visible)
 					{
 						return EVisibility::Visible;
 					}
@@ -1886,46 +1886,8 @@ void UToolMenus::PopulateToolBarBuilderWithEntry(
 		ToolBarBuilder.BeginStyleOverride(OverrideStyleName);
 	}
 
-	// If a section is invisible, then nothing is visible
-	TAttribute<EVisibility> Visibility = Section.Visibility;
-	
-	if (bIsRaisingToTopLevel)
-	{
-		// Only customize visibility if this is bound.
-		// A static value of `true` got this block here in the first place
-		if (Block.ShowInToolbarTopLevel.IsBound())
-		{
-			Visibility = UE::ToolMenus::Private::CombineVisibilityAttributes(
-				Visibility,
-				TAttribute<EVisibility>::CreateLambda([ShowInToolbarTopLevel = Block.ShowInToolbarTopLevel]
-				{
-					return ShowInToolbarTopLevel.Get() ? EVisibility::Visible : EVisibility::Collapsed;
-				})
-			);	
-		}
-		
-		if (Block.ScriptObject)
-		{
-			// Allow scripts to override ShowInToolbarTopLevel.
-			if (const TAttribute<bool> ScriptShowInToolbarTopLevel = Block.ScriptObject->CreateShowInToolbarTopLevelAttribute(MenuData->Context);
-				ScriptShowInToolbarTopLevel.IsSet())
-			{
-				Visibility = UE::ToolMenus::Private::CombineVisibilityAttributes(
-					TAttribute<EVisibility>::CreateLambda([ScriptShowInToolbarTopLevel]
-					{
-						return ScriptShowInToolbarTopLevel.Get() ? EVisibility::Visible : EVisibility::Collapsed;
-					}),
-					Visibility
-				);
-			}
-		}
-	}
-	
-	// Block level visibility applies last
-	Visibility = UE::ToolMenus::Private::CombineVisibilityAttributes(
-		Visibility,
-		Block.Visibility
-	);
+	constexpr bool bEmbedActionOrCommand = false;
+	const TAttribute<EVisibility> Visibility = CalculateToolbarVisibility(MenuData, Section, Block, bIsRaisingToTopLevel, bEmbedActionOrCommand);
 
 	const FUIAction UIAction = [&Block, MenuData]()
 		{
@@ -2584,6 +2546,100 @@ void UToolMenus::ExecuteStringCommand(const FToolMenuStringCommand StringCommand
 	}
 }
 
+TAttribute<EVisibility> UToolMenus::CalculateToolbarVisibility(UToolMenu* Menu, FToolMenuSection& Section, FToolMenuEntry& Entry, bool bIsRaisingToTopLevel, bool bEmbedActionOrCommand)
+{
+	// Single values get returned directly. Allow that to cause zero allocations.
+	TArray<FToolMenuVisibilityChoice, TInlineAllocator<1>> Visibilities;
+
+	if (Section.Visibility.IsSet())
+	{
+		Visibilities.Add(Section.Visibility);
+	}
+
+	if (bIsRaisingToTopLevel)
+	{
+		FToolMenuVisibilityChoice ShowInTopLevel = Entry.ShowInToolbarTopLevel;
+		if (Entry.ScriptObject)
+		{
+			// Allow scripts to override ShowInToolbarTopLevel.
+			if (const TAttribute<bool> ScriptShowInToolbarTopLevel = Entry.ScriptObject->CreateShowInToolbarTopLevelAttribute(Menu->Context);
+				ScriptShowInToolbarTopLevel.IsSet())
+			{
+				ShowInTopLevel = UE::ToolMenus::Private::CombineVisibility(
+					ScriptShowInToolbarTopLevel,
+					ShowInTopLevel
+				);
+			}
+		}
+		
+		if (ShowInTopLevel.IsSet())
+		{
+			Visibilities.Add(ShowInTopLevel);
+		}
+	}
+	
+	if (Entry.Visibility.IsSet())
+	{
+		Visibilities.Add(Entry.Visibility);
+	}
+	
+	if (bEmbedActionOrCommand)
+	{
+		if (Entry.Command.IsValid() && !Entry.IsCommandKeybindOnly())
+		{
+			TSharedPtr<const FUICommandList> CommandListForAction;
+			if (const FUIAction* FoundAction = Entry.GetActionForCommand(Menu->Context, CommandListForAction))
+			{
+				if (FoundAction->IsActionVisibleDelegate.IsBound())
+				{
+					Visibilities.Add(FoundAction->IsActionVisibleDelegate);
+				}
+			}
+		}
+	
+		const FUIAction UIAction = [&Entry, Menu]()
+		{
+			if (Entry.ToolBarData.ActionOverride.IsSet())
+			{
+				return UToolMenus::ConvertUIAction(Entry.ToolBarData.ActionOverride.GetValue(), Menu->Context);
+			}
+			else
+			{
+				return UToolMenus::ConvertUIAction(Entry, Menu->Context);
+			}
+		}();
+		
+		if (UIAction.IsActionVisibleDelegate.IsBound())
+		{
+			Visibilities.Add(UIAction.IsActionVisibleDelegate);
+		}
+	}
+	
+	if (Visibilities.Num() == 0)
+	{
+		return TAttribute<EVisibility>();
+	}
+	
+	if (Visibilities.Num() == 1)
+	{
+		return Visibilities[0];
+	}
+	
+	return TAttribute<EVisibility>::CreateLambda([Visibilities]
+	{
+		for (const FToolMenuVisibilityChoice& VisibilityChoice : Visibilities)
+		{
+			const EVisibility Visibility = VisibilityChoice.Get();
+			if (Visibility != EVisibility::Visible)
+			{
+				return Visibility;
+			}
+		}
+		
+		return EVisibility::Visible;
+	});
+}
+
 UToolMenu* UToolMenus::FindSubMenuToGenerateWith(const FName InParentName, const FName InChildName)
 {
 	FName BaseName = InParentName;
diff --git a/Engine/Source/Developer/ToolMenus/Public/ToolMenuEntry.h b/Engine/Source/Developer/ToolMenus/Public/ToolMenuEntry.h
--- a/Engine/Source/Developer/ToolMenus/Public/ToolMenuEntry.h
+++ b/Engine/Source/Developer/ToolMenus/Public/ToolMenuEntry.h
@@ -124,6 +124,32 @@ struct FToolMenuCustomWidgetContext
 	FName StyleName;
 };
 
+/**
+ * A convenience wrapper for multiple ways of delivering a visibility value.
+ * Use by setting directly to TAttribute<EVisibility>, TAttribute<bool>, or FIsActionButtonVisible.
+ * Can also assign directly to a lambda, e.g `Choice = []{ return TrueOrFalse(); };`
+ */
+struct TOOLMENUS_API FToolMenuVisibilityChoice
+{
+	FToolMenuVisibilityChoice();
+	FToolMenuVisibilityChoice(const FToolMenuVisibilityChoice& Other);
+	FToolMenuVisibilityChoice(const TAttribute<EVisibility>& InVisibility);
+	FToolMenuVisibilityChoice(const TAttribute<bool>& InIsVisible);
+	FToolMenuVisibilityChoice(const FIsActionButtonVisible& InIsActionButtonVisible);
+	
+	TAttribute<EVisibility> ToVisibilityAttribute() const;
+	operator TAttribute<EVisibility>() const { return ToVisibilityAttribute(); }
+	
+	FToolMenuVisibilityChoice& operator=(TFunction<EVisibility()> VisibilityFunc);
+	FToolMenuVisibilityChoice& operator=(TFunction<bool()> IsVisibleFunc);
+	
+	bool IsSet() const;
+	EVisibility Get() const;
+	
+private:
+	TVariant<TAttribute<EVisibility>, TAttribute<bool>, FIsActionButtonVisible> Value;
+};
+
 /**
  * Represents entries in menus such as buttons, checkboxes, and sub-menus.
  *
@@ -203,8 +229,6 @@ struct TOOLMENUS_API FToolMenuEntry
 	 */
 	void SetShowInToolbarTopLevel(TAttribute<bool> InTopLevel);
 
-	void SetVisibility(TAttribute<EVisibility> InVisibility);
-
 private:
 
 	void SetCommand(const TSharedPtr< const FUICommandInfo >& InCommand, TOptional<FName> InName, const TAttribute<FText>& InLabel, const TAttribute<FText>& InToolTip, const TAttribute<FSlateIcon>& InIcon);
@@ -261,6 +285,8 @@ public:
 	TAttribute<FSlateIcon> Icon;
 	TAttribute<FText> InputBindingLabel;
 
+	FToolMenuVisibilityChoice Visibility;
+	
 private:
 
 	friend class UToolMenus;
@@ -283,5 +309,4 @@ private:
 	bool bCommandIsKeybindOnly;
 
 	TAttribute<bool> ShowInToolbarTopLevel;
-	TAttribute<EVisibility> Visibility;
 };
diff --git a/Engine/Source/Developer/ToolMenus/Public/ToolMenuSection.h b/Engine/Source/Developer/ToolMenus/Public/ToolMenuSection.h
--- a/Engine/Source/Developer/ToolMenus/Public/ToolMenuSection.h
+++ b/Engine/Source/Developer/ToolMenus/Public/ToolMenuSection.h
@@ -133,7 +133,7 @@ public:
 	TObjectPtr<UToolMenuSectionDynamic> ToolMenuSectionDynamic;
 
 	TAttribute<FText> Label;
-	TAttribute<EVisibility> Visibility;
+	FToolMenuVisibilityChoice Visibility;
 
 	FNewSectionConstructChoice Construct;
 
diff --git a/Engine/Source/Developer/ToolMenus/Public/ToolMenus.h b/Engine/Source/Developer/ToolMenus/Public/ToolMenus.h
--- a/Engine/Source/Developer/ToolMenus/Public/ToolMenus.h
+++ b/Engine/Source/Developer/ToolMenus/Public/ToolMenus.h
@@ -433,6 +433,14 @@ private:
 	static FUIAction ConvertUIAction(const FToolUIAction& Actions, const FToolMenuContext& Context);
 	static FUIAction ConvertUIAction(const FToolDynamicUIAction& Actions, const FToolMenuContext& Context);
 	static FUIAction ConvertScriptObjectToUIAction(UToolMenuEntryScript* ScriptObject, const FToolMenuContext& Context);
+	
+	static TAttribute<EVisibility> CalculateToolbarVisibility(
+		UToolMenu* Menu,
+		FToolMenuSection& Section,
+		FToolMenuEntry& Entry,
+		bool bIsRaisingToTopLevel,
+		bool bEmbedActionOrCommand
+	);
 
 	static void ExecuteStringCommand(const FToolMenuStringCommand StringCommand, FName MenuName, const FToolMenuContext Context);
 
diff --git a/Engine/Source/Editor/LevelEditor/Private/ViewportToolbar/LevelEditorSubmenus.cpp b/Engine/Source/Editor/LevelEditor/Private/ViewportToolbar/LevelEditorSubmenus.cpp
--- a/Engine/Source/Editor/LevelEditor/Private/ViewportToolbar/LevelEditorSubmenus.cpp
+++ b/Engine/Source/Editor/LevelEditor/Private/ViewportToolbar/LevelEditorSubmenus.cpp
@@ -522,14 +522,14 @@ FToolMenuEntry CreateActorSnapCheckboxMenu()
 }
 
 // Can be used to show entries only in perspective view - specialized version of UE::UnrealEd method, for LevelViewport argument
-TAttribute<EVisibility> GetPerspectiveOnlyVisibility(const TWeakPtr<::SLevelViewport>& LevelViewportWeak)
+TAttribute<bool> GetIsPerspectiveAttribute(const TWeakPtr<::SLevelViewport>& LevelViewportWeak)
 {
 	if (TSharedPtr<::SLevelViewport> LevelViewport = LevelViewportWeak.Pin())
 	{
-		return UE::UnrealEd::GetPerspectiveOnlyVisibility(LevelViewport->GetViewportClient());
+		return UE::UnrealEd::GetIsPerspectiveAttribute(LevelViewport->GetViewportClient());
 	}
 
-	return EVisibility::Collapsed;
+	return false;
 }
 
 FToolMenuEntry CreatePilotSubmenu(TWeakPtr<::SLevelViewport> LevelViewportWeak)
@@ -599,7 +599,7 @@ FToolMenuEntry CreatePilotSubmenu(TWeakPtr<::SLevelViewport> LevelViewportWeak)
 		FSlateIcon(FAppStyle::Get().GetStyleSetName(), "LevelViewport.PilotSelectedActor")
 	);
 
-	Entry.SetVisibility(UE::LevelEditor::Private::GetPerspectiveOnlyVisibility(LevelViewportWeak));
+	Entry.Visibility = GetIsPerspectiveAttribute(LevelViewportWeak);
 
 	return Entry;
 }
@@ -641,7 +641,7 @@ FToolMenuEntry CreateCameraMovementSubmenu(const TWeakPtr<::SLevelViewport>& Lev
 		FSlateIcon(FSlateIcon(FAppStyle::GetAppStyleSetName(), "ClassIcon.CameraComponent"))
 	);
 
-	Entry.SetVisibility(GetPerspectiveOnlyVisibility(LevelViewportWeak));
+	Entry.Visibility = GetIsPerspectiveAttribute(LevelViewportWeak);
 
 	return Entry;
 }
diff --git a/Engine/Source/Editor/UnrealEd/Private/ViewportToolbar/UnrealEdViewportToolbar.cpp b/Engine/Source/Editor/UnrealEd/Private/ViewportToolbar/UnrealEdViewportToolbar.cpp
--- a/Engine/Source/Editor/UnrealEd/Private/ViewportToolbar/UnrealEdViewportToolbar.cpp
+++ b/Engine/Source/Editor/UnrealEd/Private/ViewportToolbar/UnrealEdViewportToolbar.cpp
@@ -207,14 +207,14 @@ FToolMenuEntry CreateSurfaceSnapCheckboxMenu(const FToolMenuContext& InContext)
 	if (UUnrealEdViewportToolbarContext* Context = InContext.FindContext<UUnrealEdViewportToolbarContext>())
 	{
 		TWeakObjectPtr<UUnrealEdViewportToolbarContext> WeakContext = Context;
-		Entry.SetVisibility(TAttribute<EVisibility>::CreateLambda([WeakContext]
+		Entry.Visibility = [WeakContext]
 		{
 			if (WeakContext.IsValid())
 			{
-				return WeakContext->bShowSurfaceSnap ? EVisibility::Visible : EVisibility::Hidden;
+				return WeakContext->bShowSurfaceSnap;
 			}
-			return EVisibility::Visible;
-		}));
+			return true;
+		};
 	}
 
 	Entry.InputBindingLabel = TAttribute<FText>::CreateLambda(
@@ -3114,7 +3114,7 @@ void PopulateCameraMenu(UToolMenu* InMenu, const FViewportCameraMenuOptions& InO
 			FrameEntry.Icon = FSlateIcon(FAppStyle::GetAppStyleSetName(), "Icons.FrameActor");
 		}
 
-		MovementSection.Visibility = UE::UnrealEd::GetPerspectiveOnlyVisibility(EditorViewport->GetViewportClient());
+		MovementSection.Visibility = UE::UnrealEd::GetIsPerspectiveAttribute(EditorViewport->GetViewportClient());
 	}
 
 	// View
@@ -3133,7 +3133,7 @@ void PopulateCameraMenu(UToolMenu* InMenu, const FViewportCameraMenuOptions& InO
 			ViewSection.AddEntry(CameraFOV);
 		}
 
-		ViewSection.Visibility = UE::UnrealEd::GetPerspectiveOnlyVisibility(EditorViewport->GetViewportClient());
+		ViewSection.Visibility = UE::UnrealEd::GetIsPerspectiveAttribute(EditorViewport->GetViewportClient());
 
 		if (InOptions.bShowNearAndFarPlanes)
 		{
@@ -4164,19 +4164,27 @@ FToolMenuEntry CreateCameraSpeedMenu()
 
 TAttribute<EVisibility> GetPerspectiveOnlyVisibility(const TSharedPtr<FEditorViewportClient>& InViewportClient)
 {
-	if (!InViewportClient)
+	return TAttribute<EVisibility>::CreateLambda([IsPerspective = GetIsPerspectiveAttribute(InViewportClient)]
+	{
+		return IsPerspective.Get() ? EVisibility::Visible : EVisibility::Collapsed;
+	});
+}
+
+TAttribute<bool> GetIsPerspectiveAttribute(const TSharedPtr<FEditorViewportClient>& InViewportClient)
+{
+	if (!InViewportClient.IsValid())
 	{
-		return EVisibility::Collapsed;
+		return false;
 	}
 
-	return TAttribute<EVisibility>::CreateLambda(
-		[WeakViewport = InViewportClient.ToWeakPtr()]() -> EVisibility
+	return TAttribute<bool>::CreateLambda(
+		[WeakViewport = InViewportClient.ToWeakPtr()]
 		{
 			if (const TSharedPtr<FEditorViewportClient>& ViewportClient = WeakViewport.Pin())
 			{
-				return ViewportClient->IsPerspective() ? EVisibility::Visible : EVisibility::Collapsed;
+				return ViewportClient->IsPerspective();
 			}
-			return EVisibility::Collapsed;
+			return false;
 		}
 	);
 }
diff --git a/Engine/Source/Editor/UnrealEd/Public/ViewportToolbar/UnrealEdViewportToolbar.h b/Engine/Source/Editor/UnrealEd/Public/ViewportToolbar/UnrealEdViewportToolbar.h
--- a/Engine/Source/Editor/UnrealEd/Public/ViewportToolbar/UnrealEdViewportToolbar.h
+++ b/Engine/Source/Editor/UnrealEd/Public/ViewportToolbar/UnrealEdViewportToolbar.h
@@ -269,6 +269,10 @@ UNREALED_API FText GetCameraSpeedTooltip();
  * Returns Visible if the provided Viewport Client is in perspective view, Collapsed if not
  */
 UNREALED_API TAttribute<EVisibility> GetPerspectiveOnlyVisibility(const TSharedPtr<FEditorViewportClient>& InViewportClient);
+/**
+ * Creates an attribute that provides whether the viewport client is in perspective view.
+ */
+UNREALED_API TAttribute<bool> GetIsPerspectiveAttribute(const TSharedPtr<FEditorViewportClient>& InViewportClient);
 
 // Camera Menu Widgets
 TSharedRef<SWidget> CreateCameraMenuWidget(const TSharedRef<SEditorViewport>& InViewport, bool bInShowExposureSettings = false);
