Hotfix of missing forward declaration of FToolMenuEntry in ToolMenuDelegates.h

#rnx

diff --git a/Engine/Source/Developer/ToolMenus/Public/ToolMenuDelegates.h b/Engine/Source/Developer/ToolMenus/Public/ToolMenuDelegates.h
--- a/Engine/Source/Developer/ToolMenus/Public/ToolMenuDelegates.h
+++ b/Engine/Source/Developer/ToolMenus/Public/ToolMenuDelegates.h
@@ -16,6 +16,7 @@
 #define UE_API TOOLMENUS_API
 
 class UToolMenu;
+struct FToolMenuEntry;
 struct FToolMenuSection;
 struct FToolMenuCustomWidgetContext;
 
