Fix an issue where toolbars with a gap in the middle could have that middle gap collapse prior to any clipping occuring, leaving space on the left side.

#rnx

diff --git a/Engine/Source/Developer/ToolMenus/Private/ToolMenus.cpp b/Engine/Source/Developer/ToolMenus/Private/ToolMenus.cpp
--- a/Engine/Source/Developer/ToolMenus/Private/ToolMenus.cpp
+++ b/Engine/Source/Developer/ToolMenus/Private/ToolMenus.cpp
@@ -2171,7 +2171,7 @@ void UToolMenus::PopulateToolBarBuilder(FToolBarBuilder& ToolBarBuilder, UToolMe
 				FMenuEntryStyleParams StyleParams;
 
 				StyleParams.HorizontalAlignment = HAlign_Right;
-				StyleParams.SizeRule = FSizeParam::ESizeRule::SizeRule_StretchContent;
+				StyleParams.SizeRule = FSizeParam::ESizeRule::SizeRule_Stretch;
 				StyleParams.FillSize = AlmostZero;
 				StyleParams.FillSizeMin = AlmostZero; // Allow the spacer to shrink to nothing if another widget needs the space.
 				StyleParams.MinimumSize = AlmostZero;
