Add support for displaying custom widgets within Tool Bar Button Blocks. This fixes an issue where Menu Items with MakeCustomWidget delegates would show up differently when displayed in a Toolbar context.

#rnx

diff --git a/Engine/Source/Developer/ToolMenus/Private/ToolMenus.cpp b/Engine/Source/Developer/ToolMenus/Private/ToolMenus.cpp
--- a/Engine/Source/Developer/ToolMenus/Private/ToolMenus.cpp
+++ b/Engine/Source/Developer/ToolMenus/Private/ToolMenus.cpp
@@ -9,6 +9,7 @@
 #include "Framework/MultiBox/MultiBox.h"
 #include "Framework/MultiBox/MultiBoxBuilder.h"
 #include "Framework/Application/SlateApplication.h"
+#include "Framework/MultiBox/SToolBarButtonBlock.h"
 #include "Internationalization/Internationalization.h"
 #include "UObject/WeakObjectPtrTemplates.h"
 
@@ -1954,7 +1955,7 @@ void UToolMenus::PopulateToolBarBuilderWithEntry(
 					*Block.Name.ToString(), *MenuData->MenuName.ToString());
 			}
 
-			ToolBarBuilder.AddToolBarButton(
+			TSharedPtr<FToolBarButtonBlock> ButtonBlock = ToolBarBuilder.AddToolBarButton(
 				Block.Command,
 				Block.Name,
 				Block.Label,
@@ -1966,6 +1967,15 @@ void UToolMenus::PopulateToolBarBuilderWithEntry(
 				ToolbarLabelOverride,
 				ActualResizeParams
 			);
+			
+			if (Block.MakeCustomWidget.IsBound())
+			{
+				FToolMenuCustomWidgetContext EntryWidgetContext;
+				TSharedRef<FMultiBox> MultiBox = ToolBarBuilder.GetMultiBox();
+				EntryWidgetContext.StyleSet = MultiBox->GetStyleSet();
+				EntryWidgetContext.StyleName = MultiBox->GetStyleName();
+				ButtonBlock->SetCustomWidget(Block.MakeCustomWidget.Execute(MenuData->Context, EntryWidgetContext));
+			}
 
 			if (bPopCommandList)
 			{
@@ -1992,7 +2002,7 @@ void UToolMenus::PopulateToolBarBuilderWithEntry(
 		}
 		else
 		{
-			ToolBarBuilder.AddToolBarButton(
+			TSharedPtr<FToolBarButtonBlock> ButtonBlock = ToolBarBuilder.AddToolBarButton(
 				UIAction,
 				Block.Name,
 				Block.Label,
@@ -2004,6 +2014,15 @@ void UToolMenus::PopulateToolBarBuilderWithEntry(
 				ToolbarLabelOverride,
 				ActualResizeParams
 			);
+			
+			if (Block.MakeCustomWidget.IsBound())
+			{
+				FToolMenuCustomWidgetContext EntryWidgetContext;
+				TSharedRef<FMultiBox> MultiBox = ToolBarBuilder.GetMultiBox();
+				EntryWidgetContext.StyleSet = MultiBox->GetStyleSet();
+				EntryWidgetContext.StyleName = MultiBox->GetStyleName();
+				ButtonBlock->SetCustomWidget(Block.MakeCustomWidget.Execute(MenuData->Context, EntryWidgetContext));
+			}
 		}
 
 		if (Block.ToolBarData.OptionsDropdownData.IsValid())
diff --git a/Engine/Source/Runtime/Slate/Private/Framework/MultiBox/MultiBoxBuilder.cpp b/Engine/Source/Runtime/Slate/Private/Framework/MultiBox/MultiBoxBuilder.cpp
--- a/Engine/Source/Runtime/Slate/Private/Framework/MultiBox/MultiBoxBuilder.cpp
+++ b/Engine/Source/Runtime/Slate/Private/Framework/MultiBox/MultiBoxBuilder.cpp
@@ -614,7 +614,7 @@ void FToolBarBuilder::SetAllowWrapButton(const TOptional<bool>& bInAllowWrapButt
 	MultiBox->bAllowWrapButton = bInAllowWrapButton;
 }
 
-void FToolBarBuilder::AddToolBarButton(
+TSharedPtr<FToolBarButtonBlock> FToolBarBuilder::AddToolBarButton(
 	const TSharedPtr<const FUICommandInfo> InCommand,
 	FName InExtensionHook,
 	const TAttribute<FText>& InLabelOverride,
@@ -653,6 +653,8 @@ void FToolBarBuilder::AddToolBarButton(
 	MultiBox->AddMultiBlock( NewToolBarButtonBlock );
 
 	ApplyHook(InExtensionHook, EExtensionHook::After);
+	
+	return NewToolBarButtonBlock;
 }
 
 void FToolBarBuilder::AddToolBarButton(const FButtonArgs& Args)
@@ -701,7 +703,7 @@ void FToolBarBuilder::AddToolBarButton(const FButtonArgs& Args)
 	ApplyHook(Args.ExtensionHook, EExtensionHook::After);
 }
 
-void FToolBarBuilder::AddToolBarButton(
+TSharedPtr<FToolBarButtonBlock> FToolBarBuilder::AddToolBarButton(
 	const FUIAction& InAction,
 	FName InExtensionHook,
 	const TAttribute<FText>& InLabelOverride,
@@ -739,6 +741,8 @@ void FToolBarBuilder::AddToolBarButton(
 	MultiBox->AddMultiBlock( NewToolBarButtonBlock );
 
 	ApplyHook(InExtensionHook, EExtensionHook::After);
+	
+	return NewToolBarButtonBlock;
 }
 
 void FToolBarBuilder::AddComboButton(
diff --git a/Engine/Source/Runtime/Slate/Private/Framework/MultiBox/SToolBarButtonBlock.cpp b/Engine/Source/Runtime/Slate/Private/Framework/MultiBox/SToolBarButtonBlock.cpp
--- a/Engine/Source/Runtime/Slate/Private/Framework/MultiBox/SToolBarButtonBlock.cpp
+++ b/Engine/Source/Runtime/Slate/Private/Framework/MultiBox/SToolBarButtonBlock.cpp
@@ -75,6 +75,11 @@ void FToolBarButtonBlock::SetCustomMenuDelegate(const FNewMenuDelegate& InCustom
 	CustomMenuDelegate = InCustomMenuDelegate;
 }
 
+void FToolBarButtonBlock::SetCustomWidget(const TSharedPtr<SWidget>& InCustomWidget)
+{
+	CustomWidget = InCustomWidget;
+}
+
 void FToolBarButtonBlock::SetOnGetMenuContent(const FOnGetContent& InOnGetMenuContent)
 {
 	OnGetMenuContent = InOnGetMenuContent;
@@ -274,7 +279,12 @@ void SToolBarButtonBlock::BuildMultiBlockWidget(const ISlateStyle* StyleSet, con
 
 	// Create the content for our button
 	TSharedRef<SWidget> ButtonContent = SNullWidget::NullWidget;
-	if (MultiBox->GetType() == EMultiBoxType::SlimHorizontalToolBar 
+	
+	if (ToolBarButtonBlock->CustomWidget)
+	{
+		ButtonContent = ToolBarButtonBlock->CustomWidget.ToSharedRef();
+	}
+	else if (MultiBox->GetType() == EMultiBoxType::SlimHorizontalToolBar 
 		|| bIsSlimHorizontalUniformToolBar
 		|| bIsSlimWrappingToolBar)
 	{
diff --git a/Engine/Source/Runtime/Slate/Public/Framework/MultiBox/MultiBoxBuilder.h b/Engine/Source/Runtime/Slate/Public/Framework/MultiBox/MultiBoxBuilder.h
--- a/Engine/Source/Runtime/Slate/Public/Framework/MultiBox/MultiBoxBuilder.h
+++ b/Engine/Source/Runtime/Slate/Public/Framework/MultiBox/MultiBoxBuilder.h
@@ -594,7 +594,7 @@ public:
 	 * @param	InToolbarLabelOverride	Optional label override for when this block appears in a toolbar. If omitted, then the action's label will be used instead.
 	 * @param	InResizeParams			Optional resize parameters to control how this entry behaves during toolbar resizing.
 	 */
-	SLATE_API void AddToolBarButton(
+	SLATE_API TSharedPtr<FToolBarButtonBlock> AddToolBarButton(
 		const TSharedPtr<const FUICommandInfo> InCommand,
 		FName InExtensionHook = NAME_None,
 		const TAttribute<FText>& InLabelOverride = TAttribute<FText>(),
@@ -621,7 +621,7 @@ public:
 	 * @param	InToolbarLabelOverride	Optional label override for when this block appears in a toolbar. If omitted, then the action's label will be used instead.
 	 * @param	InResizeParams			Optional resize parameters to control how this entry behaves during toolbar resizing.
 	 */
-	SLATE_API void AddToolBarButton(
+	SLATE_API TSharedPtr<FToolBarButtonBlock> AddToolBarButton(
 		const FUIAction& InAction,
 		FName InExtensionHook = NAME_None,
 		const TAttribute<FText>& InLabelOverride = TAttribute<FText>(),
diff --git a/Engine/Source/Runtime/Slate/Public/Framework/MultiBox/SToolBarButtonBlock.h b/Engine/Source/Runtime/Slate/Public/Framework/MultiBox/SToolBarButtonBlock.h
--- a/Engine/Source/Runtime/Slate/Public/Framework/MultiBox/SToolBarButtonBlock.h
+++ b/Engine/Source/Runtime/Slate/Public/Framework/MultiBox/SToolBarButtonBlock.h
@@ -112,6 +112,8 @@ public:
 
 	/** Set optional delegate to customize when a menu appears instead of the widget, such as in toolbars */
 	SLATE_API void SetCustomMenuDelegate(const FNewMenuDelegate& InOnFillMenuDelegate);
+	
+	SLATE_API void SetCustomWidget(const TSharedPtr<SWidget>& InCustomWidget);
 
 	SLATE_API void SetOnGetMenuContent(const FOnGetContent& OnGetMenuContent);
 
@@ -153,6 +155,8 @@ private:
 
 	TOptional< EVisibility > LabelVisibility;
 	
+	TSharedPtr<SWidget> CustomWidget;
+	
 	/** In the case where a command is not bound, the user interface action type to use.  If a command is bound, we
 		simply use the action type associated with that command. */
 	EUserInterfaceActionType UserInterfaceActionType;
