- Tool Menu Items will now override override any parent Tool Menu Items of the same name.
- Positioned Tool Menu Items that fail to be positioned will now be inserted at the end of the section.

[RN] minor Editor

diff --git a/Engine/Source/Developer/ToolMenus/Private/ToolMenus.cpp b/Engine/Source/Developer/ToolMenus/Private/ToolMenus.cpp
--- a/Engine/Source/Developer/ToolMenus/Private/ToolMenus.cpp
+++ b/Engine/Source/Developer/ToolMenus/Private/ToolMenus.cpp
@@ -820,6 +820,29 @@ void UToolMenus::AssembleMenuSection(UToolMenu* GeneratedMenu, const UToolMenu*
 			for (int32 i = 0; i < RemainingBlocks.Num(); ++i)
 			{
 				FToolMenuEntry& Block = RemainingBlocks[i];
+				
+				if (!Block.Name.IsNone())
+				{
+					// Replace blocks that have the same name. Children blocks override parent blocks.
+					int32 TargetIndex = DestSection->IndexOfBlock(Block.Name);
+					if (TargetIndex != INDEX_NONE)
+					{
+						if (Block.InsertPosition.IsDefault())
+						{
+							DestSection->Blocks[TargetIndex] = Block;
+                            RemainingBlocks.RemoveAt(i);
+                            --i;
+                            ++NumHandled;
+                            break;
+						}
+						else
+						{
+							// Allow positioning logic to run as specified on the new item
+							DestSection->Blocks.RemoveAt(TargetIndex);
+						}
+					}
+				}
+				
 				int32 DestIndex = DestSection->FindBlockInsertIndex(Block);
 				if (DestIndex != INDEX_NONE)
 				{
@@ -833,11 +856,11 @@ void UToolMenus::AssembleMenuSection(UToolMenu* GeneratedMenu, const UToolMenu*
 			}
 			if (NumHandled == 0)
 			{
-				for (const FToolMenuEntry& Block : RemainingBlocks)
-				{
-					UE_LOG(LogToolMenus, Warning, TEXT("Menu item not found: '%s' for insert: '%s'"), *Block.InsertPosition.Name.ToString(), *Block.Name.ToString());
-				}
-				break;
+				// Insert the first remaining block and try again. Later blocks may have depended on that first block.
+				FToolMenuEntry& Block = RemainingBlocks[0];
+				UE_LOG(LogToolMenus, Warning, TEXT("Menu item not found: '%s' for insert: '%s'"), *Block.InsertPosition.Name.ToString(), *Block.Name.ToString());
+				DestSection->Blocks.Add(Block);
+				RemainingBlocks.RemoveAt(0);
 			}
 		}
 	}
