[ToolMenus] Adds options for whether a Tool Menu Entry with an insert position that can't find its target should insert anyways and/or log. Defaults to both inserting & logging.

Disables logging if the ToggleViewportToolbar and ToggleHideViewportUI commands can't find the ToggleFullscreen command in the Level Editor Window menu.

[RN] Editor Minor
#FailedValidations Preflight

diff --git a/Engine/Source/Developer/ToolMenus/Private/ToolMenus.cpp b/Engine/Source/Developer/ToolMenus/Private/ToolMenus.cpp
--- a/Engine/Source/Developer/ToolMenus/Private/ToolMenus.cpp
+++ b/Engine/Source/Developer/ToolMenus/Private/ToolMenus.cpp
@@ -856,11 +856,21 @@ void UToolMenus::AssembleMenuSection(UToolMenu* GeneratedMenu, const UToolMenu*
 			}
 			if (NumHandled == 0)
 			{
-				// Insert the first remaining block and try again. Later blocks may have depended on that first block.
 				FToolMenuEntry& Block = RemainingBlocks[0];
-				UE_LOG(LogToolMenus, Warning, TEXT("Menu item not found: '%s' for insert: '%s'"), *Block.InsertPosition.Name.ToString(), *Block.Name.ToString());
-				DestSection->Blocks.Add(Block);
-				RemainingBlocks.RemoveAt(0);
+				if ((Block.InsertPosition.Fallback & EToolMenuInsertFallback::Log) == EToolMenuInsertFallback::Log)
+				{
+					UE_LOG(LogToolMenus, Warning, TEXT("Menu item not found: '%s' for insert: '%s'"), *Block.InsertPosition.Name.ToString(), *Block.Name.ToString());
+				}
+				if ((Block.InsertPosition.Fallback & EToolMenuInsertFallback::Insert) == EToolMenuInsertFallback::Insert)
+				{
+					// Insert the first remaining block and try	again. Later blocks may have depended on that first block.
+					DestSection->Blocks.Add(Block);
+					RemainingBlocks.RemoveAt(0);
+				}
+				else
+				{
+					break;	
+				}
 			}
 		}
 	}
diff --git a/Engine/Source/Developer/ToolMenus/Public/ToolMenuMisc.h b/Engine/Source/Developer/ToolMenus/Public/ToolMenuMisc.h
--- a/Engine/Source/Developer/ToolMenus/Public/ToolMenuMisc.h
+++ b/Engine/Source/Developer/ToolMenus/Public/ToolMenuMisc.h
@@ -66,6 +66,17 @@ enum class EToolMenuInsertType : uint8
 	Last
 };
 
+UENUM(BlueprintType, meta = (Bitflags, UseEnumValuesAsMaskValuesInEditor="true"))
+enum class EToolMenuInsertFallback : uint8
+{
+	None = 0 UMETA(Hidden),
+	/** Inset the item even when the target is not found */
+	Insert = 1 << 0,
+	/** Log a warning if the target is not found */
+	Log = 1 << 1,
+};
+ENUM_CLASS_FLAGS(EToolMenuInsertFallback);
+
 USTRUCT(BlueprintType)
 struct FToolMenuInsert
 {
@@ -81,6 +92,10 @@ struct FToolMenuInsert
 	// How to insert
 	UPROPERTY(EditAnywhere, BlueprintReadWrite, Category = "Tool Menus")
 	EToolMenuInsertType Position;
+	
+	// Whether failure to find a target should log a warning
+	UPROPERTY(EditAnywhere, BlueprintReadWrite, Category = "Tool Menus")
+	EToolMenuInsertFallback Fallback = EToolMenuInsertFallback::Insert | EToolMenuInsertFallback::Log;
 
 	FORCEINLINE bool operator==(const FToolMenuInsert& Other) const
 	{
diff --git a/Engine/Source/Editor/LevelEditor/Private/LevelEditorMenu.cpp b/Engine/Source/Editor/LevelEditor/Private/LevelEditorMenu.cpp
--- a/Engine/Source/Editor/LevelEditor/Private/LevelEditorMenu.cpp
+++ b/Engine/Source/Editor/LevelEditor/Private/LevelEditorMenu.cpp
@@ -289,14 +289,15 @@ void FLevelEditorMenu::RegisterLevelEditorMenus()
 				// Making sure the separator exist - the one in Window menu might be missing if not on a Win build.
 				// See FMainMenu::RegisterWindowMenu().
 				LayoutSection.AddSeparator("FullscreenSeparator");
+				
+				FToolMenuInsert Insert("ToggleFullscreen", EToolMenuInsertType::Before);
+				Insert.Fallback = EToolMenuInsertFallback::Insert;
 
 				FToolMenuEntry& ToggleViewportToolbar = LayoutSection.AddMenuEntry(FLevelViewportCommands::Get().ToggleViewportToolbar);
-				ToggleViewportToolbar.InsertPosition.Position = EToolMenuInsertType::Before;
-				ToggleViewportToolbar.InsertPosition.Name = "ToggleFullscreen";
+				ToggleViewportToolbar.InsertPosition = Insert;
 				
 				FToolMenuEntry& ToggleUIEntry = LayoutSection.AddMenuEntry(FLevelEditorCommands::Get().ToggleHideViewportUI);
-				ToggleUIEntry.InsertPosition.Position = EToolMenuInsertType::Before;
-				ToggleUIEntry.InsertPosition.Name = "ToggleFullscreen";
+				ToggleUIEntry.InsertPosition = Insert;
 			}
 		}
 
