[ToolMenus] Fix typos and missing information in comments.

#rnx

diff --git a/Engine/Source/Developer/ToolMenus/Public/ToolMenuMisc.h b/Engine/Source/Developer/ToolMenus/Public/ToolMenuMisc.h
--- a/Engine/Source/Developer/ToolMenus/Public/ToolMenuMisc.h
+++ b/Engine/Source/Developer/ToolMenus/Public/ToolMenuMisc.h
@@ -70,7 +70,7 @@ UENUM(BlueprintType, meta = (Bitflags, UseEnumValuesAsMaskValuesInEditor="true")
 enum class EToolMenuInsertFallback : uint8
 {
 	None = 0 UMETA(Hidden),
-	/** Inset the item even when the target is not found */
+	/** Insert the item even when the target is not found */
 	Insert = 1 << 0,
 	/** Log a warning if the target is not found */
 	Log = 1 << 1,
@@ -93,7 +93,7 @@ struct FToolMenuInsert
 	UPROPERTY(EditAnywhere, BlueprintReadWrite, Category = "Tool Menus")
 	EToolMenuInsertType Position;
 	
-	// Whether failure to find a target should log a warning
+	// How to handle an item when the target insertion position is not found. Defaults to inserting the item while logging that the target was not found.
 	UPROPERTY(EditAnywhere, BlueprintReadWrite, Category = "Tool Menus")
 	EToolMenuInsertFallback Fallback = EToolMenuInsertFallback::Insert | EToolMenuInsertFallback::Log;
 
