Fixed missing UToolMenu style: Uniform Horizontal Slim style. Could previously specify this but it wouldn't generate any widget

diff --git a/Engine/Source/Developer/ToolMenus/Private/ToolMenus.cpp b/Engine/Source/Developer/ToolMenus/Private/ToolMenus.cpp
--- a/Engine/Source/Developer/ToolMenus/Private/ToolMenus.cpp
+++ b/Engine/Source/Developer/ToolMenus/Private/ToolMenus.cpp
@@ -2918,7 +2918,7 @@ TSharedRef<SWidget> UToolMenus::GenerateWidget(UToolMenu* GeneratedMenu)
 	if (GeneratedMenu->IsEditing())
 	{
 		// Convert toolbar into menu during editing
-		if (GeneratedMenu->MenuType == EMultiBoxType::ToolBar || GeneratedMenu->MenuType == EMultiBoxType::VerticalToolBar || GeneratedMenu->MenuType == EMultiBoxType::UniformToolBar || GeneratedMenu->MenuType == EMultiBoxType::SlimHorizontalToolBar || GeneratedMenu->MenuType == EMultiBoxType::SlimWrappingToolBar)
+		if (GeneratedMenu->MenuType == EMultiBoxType::ToolBar || GeneratedMenu->MenuType == EMultiBoxType::VerticalToolBar || GeneratedMenu->MenuType == EMultiBoxType::UniformToolBar || GeneratedMenu->MenuType == EMultiBoxType::SlimHorizontalToolBar || GeneratedMenu->MenuType == EMultiBoxType::SlimWrappingToolBar || GeneratedMenu->MenuType == EMultiBoxType::SlimHorizontalUniformToolBar)
 		{
 			for (FToolMenuSection& Section : GeneratedMenu->Sections)
 			{
@@ -2974,7 +2974,7 @@ TSharedRef<SWidget> UToolMenus::GenerateWidget(UToolMenu* GeneratedMenu)
 		TSharedRef<SWidget> Result = MenuBarBuilder.MakeWidget();
 		GeneratedWidget = Result;
 	}
-	else if (GeneratedMenu->MenuType == EMultiBoxType::ToolBar || GeneratedMenu->MenuType == EMultiBoxType::VerticalToolBar || GeneratedMenu->MenuType == EMultiBoxType::UniformToolBar || GeneratedMenu->MenuType == EMultiBoxType::SlimHorizontalToolBar || GeneratedMenu->MenuType == EMultiBoxType::SlimWrappingToolBar)
+	else if (GeneratedMenu->MenuType == EMultiBoxType::ToolBar || GeneratedMenu->MenuType == EMultiBoxType::VerticalToolBar || GeneratedMenu->MenuType == EMultiBoxType::UniformToolBar || GeneratedMenu->MenuType == EMultiBoxType::SlimHorizontalToolBar || GeneratedMenu->MenuType == EMultiBoxType::SlimWrappingToolBar  || GeneratedMenu->MenuType == EMultiBoxType::SlimHorizontalUniformToolBar)
 	{
 		FToolBarBuilder ToolbarBuilder(GeneratedMenu->MenuType, GeneratedMenu->Context.CommandList, GeneratedMenu->MenuName, GeneratedMenu->Context.GetAllExtenders(), GeneratedMenu->bToolBarForceSmallIcons);
 		ToolbarBuilder.SetExtendersEnabled(GeneratedMenu->bExtendersEnabled);
@@ -3216,7 +3216,7 @@ bool UToolMenus::RefreshMenuWidget(const FName InName, FGeneratedToolMenuWidget&
 		PopulateMenuBarBuilder(MenuBarBuilder, GeneratedMenu);
 		MultiBoxWidget->SetMultiBox(MenuBarBuilder.GetMultiBox());
 	}
-	else if (GeneratedMenu->MenuType == EMultiBoxType::ToolBar || GeneratedMenu->MenuType == EMultiBoxType::VerticalToolBar || GeneratedMenu->MenuType == EMultiBoxType::UniformToolBar || GeneratedMenu->MenuType == EMultiBoxType::SlimHorizontalToolBar || GeneratedMenu->MenuType == EMultiBoxType::SlimWrappingToolBar)
+	else if (GeneratedMenu->MenuType == EMultiBoxType::ToolBar || GeneratedMenu->MenuType == EMultiBoxType::VerticalToolBar || GeneratedMenu->MenuType == EMultiBoxType::UniformToolBar || GeneratedMenu->MenuType == EMultiBoxType::SlimHorizontalToolBar || GeneratedMenu->MenuType == EMultiBoxType::SlimWrappingToolBar || GeneratedMenu->MenuType == EMultiBoxType::SlimHorizontalUniformToolBar)
 	{
 		FToolBarBuilder ToolbarBuilder(GeneratedMenu->MenuType, GeneratedMenu->Context.CommandList, GeneratedMenu->MenuName, GeneratedMenu->Context.GetAllExtenders(), GeneratedMenu->bToolBarForceSmallIcons);
 		ToolbarBuilder.SetExtendersEnabled(GeneratedMenu->bExtendersEnabled);
