ToolMenu :  make sure the MenuBuilder instance take in account the bShouldCloseWindowAfterMenuSelection parameter passed in AddSubMenu
- passed the block (FToolMenuEntry) to ConvertWidgetChoice so it can properly initialize the MenuBuilder in the legacy code path

#Synced-CL 48328177

diff --git a/Engine/Source/Developer/ToolMenus/Private/ToolMenus.cpp b/Engine/Source/Developer/ToolMenus/Private/ToolMenus.cpp
--- a/Engine/Source/Developer/ToolMenus/Private/ToolMenus.cpp
+++ b/Engine/Source/Developer/ToolMenus/Private/ToolMenus.cpp
@@ -195,7 +195,7 @@ public:
 		else if (Block.SubMenuData.ConstructMenu.NewToolMenuWidget.IsBound() || Block.SubMenuData.ConstructMenu.OnGetContent.IsBound())
 		{
 			// Full replacement of the widget shown when submenu is opened
-			FOnGetContent OnGetContent = UToolMenus::Get()->ConvertWidgetChoice(Block.SubMenuData.ConstructMenu, MenuData->Context);
+			FOnGetContent OnGetContent = UToolMenus::Get()->ConvertWidgetChoice(Block.SubMenuData.ConstructMenu, MenuData->Context, Block);
 			if (OnGetContent.IsBound())
 			{
 				MenuBuilder.AddWrapperSubMenu(
@@ -1605,7 +1605,7 @@ void UToolMenus::PopulateMenuBuilder(FMenuBuilder& MenuBuilder, UToolMenu* MenuD
 void UToolMenus::ExtractChildBlocksFromSubMenu(UToolMenu* ParentMenu, FToolMenuEntry& InBlock, TArray<UE::ToolMenus::FSubBlockReference>& SubMenuBlocks)
 {
 	const bool bBuiltViaDelegate =
-		ConvertWidgetChoice(InBlock.ToolBarData.ComboButtonContextMenuGenerator, ParentMenu->Context).IsBound();
+		ConvertWidgetChoice(InBlock.ToolBarData.ComboButtonContextMenuGenerator, ParentMenu->Context, InBlock).IsBound();
 
 	const bool bHasChildren = (InBlock.Type == EMultiBlockType::ToolBarComboButton
 								  || (InBlock.Type == EMultiBlockType::MenuEntry && InBlock.IsSubMenu()))
@@ -2061,7 +2061,7 @@ void UToolMenus::PopulateToolBarBuilderWithEntry(
 		if (Block.ToolBarData.OptionsDropdownData.IsValid())
 		{
 			FOnGetContent OnGetContent = ConvertWidgetChoice(
-				Block.ToolBarData.OptionsDropdownData->MenuContentGenerator, MenuData->Context);
+				Block.ToolBarData.OptionsDropdownData->MenuContentGenerator, MenuData->Context, Block);
 			ToolBarBuilder.AddComboButton(
 				Block.ToolBarData.OptionsDropdownData->Action,
 				OnGetContent,
@@ -2081,12 +2081,12 @@ void UToolMenus::PopulateToolBarBuilderWithEntry(
 	else if (Block.Type == EMultiBlockType::ToolBarComboButton || (Block.Type == EMultiBlockType::MenuEntry && Block.IsSubMenu()))
 	{
 		bool bCouldHaveChildren = false;
-		FOnGetContent OnGetContent = ConvertWidgetChoice(Block.ToolBarData.ComboButtonContextMenuGenerator, MenuData->Context);
+		FOnGetContent OnGetContent = ConvertWidgetChoice(Block.ToolBarData.ComboButtonContextMenuGenerator, MenuData->Context, Block);
 		
 		// Allow non-tool-menu-creating choices to be applied
 		if (!OnGetContent.IsBound() && !Block.SubMenuData.ConstructMenu.NewToolMenu.IsBound())
 		{
-			OnGetContent = ConvertWidgetChoice(Block.SubMenuData.ConstructMenu, MenuData->Context);
+			OnGetContent = ConvertWidgetChoice(Block.SubMenuData.ConstructMenu, MenuData->Context, Block);
 		}
 		
 		if (!OnGetContent.IsBound())
@@ -2359,7 +2359,7 @@ void UToolMenus::PopulateMenuBarBuilder(FMenuBarBuilder& MenuBarBuilder, UToolMe
 	AddReferencedContextObjects(MenuBarBuilder.GetMultiBox(), MenuData);
 }
 
-FOnGetContent UToolMenus::ConvertWidgetChoice(const FNewToolMenuChoice& Choice, const FToolMenuContext& Context) const
+FOnGetContent UToolMenus::ConvertWidgetChoice(const FNewToolMenuChoice& Choice, const FToolMenuContext& Context, const FToolMenuEntry& Block) const
 {
 	if (Choice.NewToolMenuWidget.IsBound())
 	{
@@ -2390,11 +2390,11 @@ FOnGetContent UToolMenus::ConvertWidgetChoice(const FNewToolMenuChoice& Choice,
 	}
 	else if (Choice.NewMenuLegacy.IsBound())
 	{
-		return FOnGetContent::CreateLambda([ToCall = Choice.NewMenuLegacy, Context]()
+		return FOnGetContent::CreateLambda([ToCall = Choice.NewMenuLegacy, Context, bShouldCloseWindowAfterMenuSelection = Block.bShouldCloseWindowAfterMenuSelection]()
 		{
 			if (ToCall.IsBound())
 			{
-				FMenuBuilder MenuBuilder(true, Context.CommandList, Context.GetAllExtenders());
+				FMenuBuilder MenuBuilder(bShouldCloseWindowAfterMenuSelection, Context.CommandList, Context.GetAllExtenders());
 				ToCall.Execute(MenuBuilder);
 				return MenuBuilder.MakeWidget();
 			}
diff --git a/Engine/Source/Developer/ToolMenus/Public/ToolMenus.h b/Engine/Source/Developer/ToolMenus/Public/ToolMenus.h
--- a/Engine/Source/Developer/ToolMenus/Public/ToolMenus.h
+++ b/Engine/Source/Developer/ToolMenus/Public/ToolMenus.h
@@ -431,7 +431,7 @@ private:
 
 	UE_API TSharedRef<SWidget> GenerateToolbarComboButtonMenu(TWeakObjectPtr<UToolMenu> InParent, const FName InBlockName);
 
-	UE_API FOnGetContent ConvertWidgetChoice(const FNewToolMenuChoice& Choice, const FToolMenuContext& Context) const;
+	UE_API FOnGetContent ConvertWidgetChoice(const FNewToolMenuChoice& Choice, const FToolMenuContext& Context, const FToolMenuEntry& Block) const;
 
 	/** Converts a string command to a FUIAction */
 	static UE_API FUIAction ConvertUIAction(const FToolMenuEntry& Block, const FToolMenuContext& Context);
